[MLDeformer] A complete redesign, turning this into a whole framework that can handle different models.

diff --git a/Engine/Plugins/Experimental/MLDeformer/MLDeformerFramework/Source/MLDeformerFrameworkEditor/Public/MLDeformerEditorModel.h b/Engine/Plugins/Experimental/MLDeformer/MLDeformerFramework/Source/MLDeformerFrameworkEditor/Public/MLDeformerEditorModel.h
new file mode 100644
--- /dev/null
+++ b/Engine/Plugins/Experimental/MLDeformer/MLDeformerFramework/Source/MLDeformerFrameworkEditor/Public/MLDeformerEditorModel.h
@@ -0,0 +1,210 @@
+// Copyright Epic Games, Inc. All Rights Reserved.
+
+#pragma once
+
+#include "CoreMinimal.h"
+#include "UObject/GCObject.h"
+#include "UObject/ObjectMacros.h"
+#include "UObject/ObjectPtr.h"
+#include "MLDeformerEditorActor.h"
+
+class UMLDeformerModel;
+class UMLDeformerInputInfo;
+class AAnimationEditorPreviewActor;
+class IPersonaPreviewScene;
+class FEditorViewportClient;
+class FSceneView;
+class FViewport;
+class FPrimitiveDrawInterface;
+class UWorld;
+class UMeshDeformer;
+class USkeletalMesh;
+class UAnimSequence;
+class UMaterial;
+class UGeometryCache;
+class UNeuralNetwork;
+
+/** Training process return codes. */
+UENUM()
+enum class ETrainingResult : uint8
+{
+	Success = 0,	/** The training successfully finished. */
+	Aborted,		/** The user has aborted the training process. */
+	AbortedCantUse,	/** The user has aborted the training process and we can't use the resulting network. */
+	FailOnData,		/** The input or output data to the network has issues, which means we cannot train. */
+	FailPythonError,/** The python script has some error (see output log). */
+	FailMissingPythonClass,	/** No training model class was implemented in Python, including a valid train function. */
+};
+
+namespace UE::MLDeformer
+{
+	class FMLDeformerEditorToolkit;
+	class FMLDeformerEditorActor;
+	class FMLDeformerSampler;
+
+	/** The base class for the editor side of an UMLDeformerModel. */
+	class MLDEFORMERFRAMEWORKEDITOR_API FMLDeformerEditorModel
+		: public FGCObject
+	{
+	public:
+		/** 
+		 * The editor model initialization settings.
+		 * This is used in the Init call.
+		 */
+		struct MLDEFORMERFRAMEWORKEDITOR_API InitSettings
+		{
+			FMLDeformerEditorToolkit* Editor = nullptr;
+			UMLDeformerModel* Model = nullptr;
+		};
+
+		virtual ~FMLDeformerEditorModel();
+
+		// Required overrides.
+		virtual int32 GetNumFrames() const PURE_VIRTUAL(FMLDeformerEditorModel::GetNumFrames, return 0;);
+		virtual ETrainingResult Train() PURE_VIRTUAL(FMLDeformerEditorModel::Train, return ETrainingResult::Success;);
+		virtual double GetTimeAtFrame(int32 FrameNumber) const PURE_VIRTUAL(FMLDeformerEditorModel::GetTimeAtFrame, return 0.0;);
+		virtual int32 GetFrameAtTime(double TimeInSeconds) const PURE_VIRTUAL(FMLDeformerEditorModel::GetFrameAtTime, return 0;);
+
+		// Optional overrides.
+		virtual void Init(const InitSettings& Settings);
+		virtual void CreateActors(const TSharedRef<IPersonaPreviewScene>& InPersonaPreviewScene);
+		virtual void ClearWorld();
+		virtual FMLDeformerEditorActor* CreateEditorActor(const FMLDeformerEditorActor::FConstructSettings& Settings) const;
+		virtual FMLDeformerSampler* CreateSampler() const;
+		virtual void Tick(FEditorViewportClient* ViewportClient, float DeltaTime);
+		virtual void CreateTrainingLinearSkinnedActor(const TSharedRef<IPersonaPreviewScene>& InPersonaPreviewScene);
+		virtual void CreateTestLinearSkinnedActor(UWorld* World);
+		virtual void CreateTestMLDeformedActor(UWorld* World);
+		virtual void CreateTrainingGroundTruthActor(UWorld* World) {}
+		virtual void CreateTestGroundTruthActor(UWorld* World) {}
+		virtual void OnTrainingDataFrameChanged();
+		virtual void UpdateActorTransforms();
+		virtual void UpdateActorVisibility();
+		virtual void UpdateLabels();
+		virtual void OnInputAssetsChanged();
+		virtual void OnPostInputAssetChanged();
+		virtual void HandleDefaultPropertyChanges(FPropertyChangedEvent& PropertyChangedEvent);
+		virtual void OnPlayButtonPressed();
+		virtual void OnPreTraining() {}
+		virtual void OnPostTraining(ETrainingResult TrainingResult);
+		virtual void OnTrainingAborted() {}
+		virtual bool IsPlayingAnim() const;
+		virtual double CalcTimelinePosition() const;
+		virtual void OnTimeSliderScrubPositionChanged(double NewScrubTime, bool bIsScrubbing);
+		virtual void UpdateTestAnimPlaySpeed();
+		virtual void ClampCurrentFrameIndex();
+		virtual int32 GetNumFramesForTraining() const;
+		virtual void SetTrainingFrame(int32 FrameNumber);
+		virtual void Render(const FSceneView* View, FViewport* Viewport, FPrimitiveDrawInterface* PDI);
+		virtual void UpdateIsReadyForTrainingState() { bIsReadyForTraining = false; }
+		virtual FText GetOverlayText() const;
+		virtual void InitInputInfo(UMLDeformerInputInfo* InputInfo);
+		virtual void RefreshMLDeformerComponents();
+		virtual void CreateHeatMapMaterial();
+		virtual void CreateHeatMapDeformerGraph();
+		virtual void CreateHeatMapAssets();
+		virtual void SetHeatMapMaterialEnabled(bool bEnabled);
+		virtual UMeshDeformer* LoadDefaultDeformerGraph();
+		virtual void SetDefaultDeformerGraphIfNeeded();
+		virtual void UpdateDeformerGraph();
+		virtual void SetComputeGraphDataProviders() const;
+		virtual void SampleDeltas();
+		virtual bool LoadTrainedNetwork() const;
+		virtual bool IsTrained() const;
+		virtual FMLDeformerEditorActor* GetTimelineEditorActor() const;
+		virtual FString GetHeatMapMaterialPath() const;
+		virtual FString GetHeatMapDeformerGraphPath() const;
+		virtual FString GetDefaultDeformerGraphAssetPath() const;
+		virtual FString GetTrainedNetworkOnnxFile() const;
+
+		FMLDeformerEditorToolkit* GetEditor() const { return Editor; }
+		UMLDeformerModel* GetModel() const { return Model.Get(); }
+		UWorld* GetWorld() const;
+		const TArray<FMLDeformerEditorActor*>& GetEditorActors() const { return EditorActors; }
+		FMLDeformerEditorActor* FindEditorActor(int32 ActorTypeID) const;
+		bool IsReadyForTraining() const { return bIsReadyForTraining; }
+		FMLDeformerSampler* GetSampler() const { return Sampler; }
+		void SetDataNormalized(bool IsNormalized) { bIsDataNormalized = IsNormalized; }
+		bool IsDataNormalized() const { return bIsDataNormalized; }
+
+		FText GetBaseAssetChangedErrorText() const;
+		FText GetVertexMapChangedErrorText() const;
+		FText GetInputsErrorText() const;
+		FText GetIncompatibleSkeletonErrorText(USkeletalMesh* InSkelMesh, UAnimSequence* InAnimSeq) const;
+		FText GetSkeletalMeshNeedsReimportErrorText() const;
+		FText GetTargetAssetChangedErrorText() const;
+
+		UMLDeformerInputInfo* GetEditorInputInfo() const { return EditorInputInfo.Get(); }	
+		void UpdateEditorInputInfo();
+
+		void TriggerInputAssetChanged(bool RefreshVizSettings=false);
+
+		void InitBoneIncludeListToAnimatedBonesOnly();
+		void InitCurveIncludeListToAnimatedCurvesOnly();
+
+		UNeuralNetwork* LoadNeuralNetworkFromOnnx(const FString& Filename) const;
+		int32 GetCurrentTrainingFrame() const { return CurrentTrainingFrame; }
+		void CheckTrainingDataFrameChanged();
+
+	protected:
+		void DeleteEditorActors();
+		bool IsEditorReadyForTrainingBasicChecks();
+
+	protected:
+		/** The runtime model associated with this editor model. */
+		TObjectPtr<UMLDeformerModel> Model = nullptr;
+
+		/** The set of actors that can appear inside the editor viewport. */
+		TArray<FMLDeformerEditorActor*> EditorActors;
+
+		/** A pointer to the editor toolkit. */
+		FMLDeformerEditorToolkit* Editor = nullptr;
+
+		/** A pointer to the sampler. */
+		FMLDeformerSampler* Sampler = nullptr;
+
+		/**
+		 * The input info as currently setup in the editor.
+		 * This is different from the runtime model's input info, as that is the one that was used to train with.
+		 */
+		TObjectPtr<UMLDeformerInputInfo> EditorInputInfo = nullptr;
+
+		/** The heatmap material. */
+		TObjectPtr<UMaterial> HeatMapMaterial = nullptr;
+
+		/** The heatmap deformer graph. */
+		TObjectPtr<UMeshDeformer> HeatMapDeformerGraph = nullptr;
+
+		/** The current training frame. */
+		int32 CurrentTrainingFrame = -1;
+
+		/**
+		 * Are we ready for training? 
+		 * The training button in the editor will be enabled or disabled based on this on default.
+		 */
+		bool bIsReadyForTraining = false;	
+
+		/** Did we already normalize the data needed for training? */
+		bool bIsDataNormalized = false;
+	};
+
+	template<class TrainingModelClass>
+	ETrainingResult TrainModel(FMLDeformerEditorModel* EditorModel)
+	{
+		// Find class, which will include our Python class, generated from the python script.
+		TArray<UClass*> TrainingModels;
+		GetDerivedClasses(TrainingModelClass::StaticClass(), TrainingModels);
+		if (TrainingModels.IsEmpty())
+		{
+			// We didn't define a derived class in Python.
+			return ETrainingResult::FailMissingPythonClass;
+		}
+
+		// Perform the training.
+		// This will trigger the Python class train function to be called.
+		TrainingModelClass* TrainingModel = Cast<TrainingModelClass>(TrainingModels.Last()->GetDefaultObject());
+		TrainingModel->Init(EditorModel);
+		const int32 ReturnCode = TrainingModel->Train();
+		return static_cast<ETrainingResult>(ReturnCode);
+	}
+}	// namespace UE::MLDeformer
diff --git a/Engine/Plugins/Experimental/MLDeformer/MLDeformerFramework/Source/MLDeformerFrameworkEditor/Public/MLDeformerEditorModule.h b/Engine/Plugins/Experimental/MLDeformer/MLDeformerFramework/Source/MLDeformerFrameworkEditor/Public/MLDeformerEditorModule.h
new file mode 100644
--- /dev/null
+++ b/Engine/Plugins/Experimental/MLDeformer/MLDeformerFramework/Source/MLDeformerFrameworkEditor/Public/MLDeformerEditorModule.h
@@ -0,0 +1,30 @@
+// Copyright Epic Games, Inc. All Rights Reserved.
+
+#pragma once
+
+#include "CoreMinimal.h"
+#include "Modules/ModuleInterface.h"
+#include "MLDeformerModelRegistry.h"
+
+namespace UE::MLDeformer
+{
+	class FMLDeformerAssetActions;
+
+	class MLDEFORMERFRAMEWORKEDITOR_API FMLDeformerEditorModule
+		: public IModuleInterface
+	{
+	public:
+		// IModuleInterface overrides.
+		void StartupModule() override;
+		void ShutdownModule() override;
+		// ~END IModuleInterface overrides.
+
+		FMLDeformerEditorModelRegistry& GetModelRegistry() { return ModelRegistry; }
+		const FMLDeformerEditorModelRegistry& GetModelRegistry() const { return ModelRegistry; }
+
+	private:
+		FMLDeformerEditorModelRegistry ModelRegistry;
+		TSharedPtr<FMLDeformerAssetActions> MLDeformerAssetActions;
+	};
+
+}	// namespace UE::MLDeformer
diff --git a/Engine/Plugins/Experimental/MLDeformer/MLDeformerFramework/Source/MLDeformerFrameworkEditor/Public/MLDeformerEditorStyle.h b/Engine/Plugins/Experimental/MLDeformer/MLDeformerFramework/Source/MLDeformerFrameworkEditor/Public/MLDeformerEditorStyle.h
new file mode 100644
--- /dev/null
+++ b/Engine/Plugins/Experimental/MLDeformer/MLDeformerFramework/Source/MLDeformerFrameworkEditor/Public/MLDeformerEditorStyle.h
@@ -0,0 +1,19 @@
+// Copyright Epic Games, Inc. All Rights Reserved.
+
+#pragma once
+
+#include "Styling/SlateStyle.h"
+
+namespace UE::MLDeformer
+{
+	class MLDEFORMERFRAMEWORKEDITOR_API FMLDeformerEditorStyle
+		: public FSlateStyleSet
+	{
+	public:
+		FMLDeformerEditorStyle();
+		~FMLDeformerEditorStyle();
+
+		static FMLDeformerEditorStyle& Get();
+	};
+
+}	// namespace UE::MLDeformer
diff --git a/Engine/Plugins/Experimental/MLDeformer/MLDeformerFramework/Source/MLDeformerFrameworkEditor/Public/MLDeformerEditorToolkit.h b/Engine/Plugins/Experimental/MLDeformer/MLDeformerFramework/Source/MLDeformerFrameworkEditor/Public/MLDeformerEditorToolkit.h
new file mode 100644
--- /dev/null
+++ b/Engine/Plugins/Experimental/MLDeformer/MLDeformerFramework/Source/MLDeformerFrameworkEditor/Public/MLDeformerEditorToolkit.h
@@ -0,0 +1,139 @@
+// Copyright Epic Games, Inc. All Rights Reserved.
+
+#pragma once
+
+#include "CoreMinimal.h"
+#include "MLDeformerTrainingModel.h"
+#include "MLDeformerEditorModel.h"
+#include "MLDeformerVizSettings.h"
+#include "IHasPersonaToolkit.h"
+#include "IPersonaPreviewScene.h"
+#include "IPersonaViewport.h"
+#include "PersonaAssetEditorToolkit.h"
+#include "EditorUndoClient.h"
+#include "Math/Color.h"
+#include "Widgets/Notifications/SNotificationList.h"
+
+class IDetailsView;
+class UMLDeformerAsset;
+class SSimpleTimeSlider;
+
+namespace UE::MLDeformer
+{
+	namespace MLDeformerEditorModes
+	{
+		extern const FName Editor;
+	}
+
+	class MLDEFORMERFRAMEWORKEDITOR_API FMLDeformerEditorToolkit :
+		public FPersonaAssetEditorToolkit,
+		public IHasPersonaToolkit,
+		public FGCObject,
+		public FEditorUndoClient,
+		public FTickableEditorObject
+	{
+	public:
+		friend class FMLDeformerApplicationMode;
+		friend struct FMLDeformerVizSettingsTabSummoner;
+
+		/** Initialize the asset editor. This will register the application mode, init the preview scene, etc. */
+		void InitAssetEditor(
+			const EToolkitMode::Type Mode,
+			const TSharedPtr<IToolkitHost>& InitToolkitHost,
+			UMLDeformerAsset* InDeformerAsset);
+
+		// FAssetEditorToolkit overrides.
+		virtual void RegisterTabSpawners(const TSharedRef<FTabManager>& InTabManager) override;
+		virtual void UnregisterTabSpawners(const TSharedRef<FTabManager>& InTabManager) override;
+		virtual FName GetToolkitFName() const override;
+		virtual FText GetBaseToolkitName() const override;
+		virtual FText GetToolkitName() const override;
+		virtual FLinearColor GetWorldCentricTabColorScale() const override;
+		virtual FString GetWorldCentricTabPrefix() const override;
+		// ~END FAssetEditorToolkit overrides.
+
+		// FGCObject overrides.
+		virtual FString GetReferencerName() const override { return TEXT("FMLDeformerEditorToolkit"); }
+		virtual void AddReferencedObjects(FReferenceCollector& Collector) override;
+		// ~END FGCObject overrides.
+
+		// FTickableEditorObject overrides.
+		virtual void Tick(float DeltaTime) override {};
+		virtual ETickableTickType GetTickableTickType() const override { return ETickableTickType::Always; }
+		virtual TStatId GetStatId() const override;
+		// ~END FTickableEditorObject overrides.
+
+		// IHasPersonaToolkit overrides.
+		virtual TSharedRef<IPersonaToolkit> GetPersonaToolkit() const override;
+		IPersonaToolkit* GetPersonaToolkitPointer() const { return PersonaToolkit.Get(); }
+		// ~END IHasPersonaToolkit overrides.
+
+		void SetVizSettingsDetailsView(TSharedPtr<IDetailsView> InDetailsView) { VizSettingsDetailsView = InDetailsView; }
+
+		IDetailsView* GetModelDetailsView() const;
+		IDetailsView* GetVizSettingsDetailsView() const;
+
+		void SetTimeSlider(TSharedPtr<SSimpleTimeSlider> InTimeSlider);
+		SSimpleTimeSlider* GetTimeSlider() const;
+
+		UMLDeformerAsset* GetDeformerAsset() const { return DeformerAsset.Get(); }
+		FMLDeformerEditorModel* GetActiveModel() { return ActiveModel.Get(); }
+
+		double CalcTimelinePosition() const;
+		void OnTimeSliderScrubPositionChanged(double NewScrubTime, bool bIsScrubbing);
+		void UpdateTimeSliderRange();
+		void SetTimeSliderRange(double StartTime, double EndTime);
+
+	private:
+		/* Toolbar related. */
+		void ExtendToolbar();
+		void FillToolbar(FToolBarBuilder& ToolbarBuilder);
+
+		/** Preview scene setup. */
+		void HandlePreviewSceneCreated(const TSharedRef<IPersonaPreviewScene>& InPersonaPreviewScene);
+		void HandleViewportCreated(const TSharedRef<IPersonaViewport>& InPersonaViewport);
+		void HandleDetailsCreated(const TSharedRef<class IDetailsView>& InDetailsView);
+		void OnFinishedChangingDetails(const FPropertyChangedEvent& PropertyChangedEvent);
+
+		/** Helpers. */
+		void ShowNotification(const FText& Message, SNotificationItem::ECompletionState State, bool PlaySound) const;
+		FText GetOverlayText() const;
+		void OnSwitchedVisualizationMode();
+		bool HandleTrainingResult(ETrainingResult TrainingResult, double TrainingDuration);
+		void OnModelChanged(int Index);
+		void OnVizModeChanged(EMLDeformerVizMode Mode);
+		FText GetActiveModelName() const;
+		FText GetCurrentVizModeName() const;
+		FText GetVizModeName(EMLDeformerVizMode Mode) const;
+		void ZoomOnActors();
+		void ShowNoModelsWarningIfNeeded();
+
+		TSharedRef<SWidget> GenerateModelButtonContents(TSharedRef<FUICommandList> InCommandList);
+		TSharedRef<SWidget> GenerateVizModeButtonContents(TSharedRef<FUICommandList> InCommandList);
+
+	private:
+		/** The persona toolkit. */	
+		TSharedPtr<IPersonaToolkit> PersonaToolkit;
+
+		/** Model details view. */
+		TSharedPtr<IDetailsView> ModelDetailsView;
+
+		/** Model viz settings details view. */
+		TSharedPtr<IDetailsView> VizSettingsDetailsView;
+
+		/** The timeline slider widget. */
+		TSharedPtr<SSimpleTimeSlider> TimeSlider;
+
+		/** The currently active editor model. */
+		TSharedPtr<FMLDeformerEditorModel> ActiveModel;
+
+		// Persona viewport.
+		TSharedPtr<IPersonaViewport> PersonaViewport;
+
+		/** The ML Deformer Asset. */
+		TObjectPtr<UMLDeformerAsset> DeformerAsset;
+
+		/** Has the asset editor been initialized? */
+		bool bIsInitialized = false;
+	};
+}	// namespace UE::MLDeformer
diff --git a/Engine/Plugins/Experimental/MLDeformer/MLDeformerFramework/Source/MLDeformerFrameworkEditor/Public/MLDeformerGeomCacheSampler.h b/Engine/Plugins/Experimental/MLDeformer/MLDeformerFramework/Source/MLDeformerFrameworkEditor/Public/MLDeformerGeomCacheSampler.h
new file mode 100644
--- /dev/null
+++ b/Engine/Plugins/Experimental/MLDeformer/MLDeformerFramework/Source/MLDeformerFrameworkEditor/Public/MLDeformerGeomCacheSampler.h
@@ -0,0 +1,57 @@
+// Copyright Epic Games, Inc. All Rights Reserved.
+
+#pragma once
+
+#include "CoreMinimal.h"
+#include "MLDeformerSampler.h"
+#include "MLDeformerModel.h"
+#include "MLDeformerGeomCacheHelpers.h"
+#include "UObject/ObjectPtr.h"
+#include "GeometryCacheMeshData.h"
+
+class UGeometryCacheComponent;
+class UGeometryCache;
+
+namespace UE::MLDeformer
+{
+	DECLARE_DELEGATE_RetVal(UGeometryCache*, FMLDeformerGetGeomCacheEvent)
+
+	/**
+	 * The input data sampler.
+	 * This class can sample bone rotations, curve values and vertex deltas.
+	 * It does this by creating two temp actors, one with skeletal mesh component and one with geom cache component.
+	 */
+	class MLDEFORMERFRAMEWORKEDITOR_API FMLDeformerGeomCacheSampler
+		: public FMLDeformerSampler
+	{
+	public:
+		// FVertexDeltaSampler overrides.
+		virtual void Sample(int32 AnimFrameIndex) override;
+		virtual void RegisterTargetComponents() override;
+		virtual float GetTimeAtFrame(int32 InAnimFrameIndex) const override;
+		// ~END FVertexDeltaSampler overrides.
+
+		const TArray<FString>& GetFailedImportedMeshNames() const { return FailedImportedMeshNames; }
+
+		FMLDeformerGetGeomCacheEvent& OnGetGeometryCache() { return GetGeometryCacheEvent; }
+
+	protected:
+		void CalculateVertexDeltas(const TArray<FVector3f>& SkinnedPositions, float DeltaCutoffLength, TArray<float>& OutVertexDeltas);
+
+	protected:
+		/** The geometry cache component used to sample the geometry cache. */
+		TObjectPtr<UGeometryCacheComponent> GeometryCacheComponent = nullptr;
+
+		/** Maps skeletal meshes imported meshes to geometry tracks. */
+		TArray<FMLDeformerGeomCacheMeshMapping> MeshMappings;
+
+		/** The geometry cache mesh data reusable buffers. One for each MeshMapping.*/
+		TArray<FGeometryCacheMeshData> GeomCacheMeshDatas;
+
+		/** Imported mesh names in the skeletal mesh for which no geom cache track could be found. */
+		TArray<FString> FailedImportedMeshNames; 
+
+		/** The function that grabs the geometry cache. */
+		FMLDeformerGetGeomCacheEvent GetGeometryCacheEvent;
+	};
+}	// namespace UE::MLDeformer
diff --git a/Engine/Plugins/Experimental/MLDeformer/MLDeformerFramework/Source/MLDeformerFrameworkEditor/Public/MLDeformerModelDetails.h b/Engine/Plugins/Experimental/MLDeformer/MLDeformerFramework/Source/MLDeformerFrameworkEditor/Public/MLDeformerModelDetails.h
new file mode 100644
--- /dev/null
+++ b/Engine/Plugins/Experimental/MLDeformer/MLDeformerFramework/Source/MLDeformerFrameworkEditor/Public/MLDeformerModelDetails.h
@@ -0,0 +1,55 @@
+// Copyright Epic Games, Inc. All Rights Reserved.
+
+#pragma once
+
+#include "CoreMinimal.h"
+#include "IDetailCustomization.h"
+#include "AssetRegistry/AssetData.h"
+#include "Input/Reply.h"
+
+class IDetailLayoutBuilder;
+class USkeleton;
+class IDetailCategoryBuilder;
+class UMLDeformerModel;
+
+namespace UE::MLDeformer
+{
+	class FMLDeformerEditorModel;
+
+	class MLDEFORMERFRAMEWORKEDITOR_API FMLDeformerModelDetails
+		: public IDetailCustomization
+	{
+	public:
+		// ILayoutDetails overrides.
+		virtual void CustomizeDetails(IDetailLayoutBuilder& DetailBuilder) override;
+		// ~END ILayoutDetails overrides.
+
+		virtual bool UpdateMemberPointers(const TArray<TWeakObjectPtr<UObject>>& Objects);
+		virtual void CreateCategories();
+		virtual void AddTargetMesh() {}
+		virtual void AddBaseMeshErrors() {}
+		virtual void AddAnimSequenceErrors() {}
+		virtual void AddTrainingInputErrors() {}
+		virtual void AddBoneInputErrors() {}
+		virtual void AddCurveInputErrors() {}
+		virtual void AddCacheSizeErrors() {}
+
+	protected:
+		bool FilterAnimSequences(const FAssetData& AssetData, USkeleton* Skeleton);
+		FReply OnFilterAnimatedBonesOnly() const;
+		FReply OnFilterAnimatedCurvesOnly() const;
+
+
+	protected:
+		/** Associated detail layout builder. */
+		IDetailLayoutBuilder* DetailLayoutBuilder = nullptr;
+
+		TObjectPtr<UMLDeformerModel> Model = nullptr;
+		FMLDeformerEditorModel* EditorModel = nullptr;
+
+		IDetailCategoryBuilder* BaseMeshCategoryBuilder = nullptr;
+		IDetailCategoryBuilder* TargetMeshCategoryBuilder = nullptr;
+		IDetailCategoryBuilder* InputOutputCategoryBuilder = nullptr;
+		IDetailCategoryBuilder* SettingsCategoryBuilder = nullptr;
+	};
+}	// namespace UE::VertexDeltaModel
diff --git a/Engine/Plugins/Experimental/MLDeformer/MLDeformerFramework/Source/MLDeformerFrameworkEditor/Public/MLDeformerModelRegistry.h b/Engine/Plugins/Experimental/MLDeformer/MLDeformerFramework/Source/MLDeformerFrameworkEditor/Public/MLDeformerModelRegistry.h
new file mode 100644
--- /dev/null
+++ b/Engine/Plugins/Experimental/MLDeformer/MLDeformerFramework/Source/MLDeformerFrameworkEditor/Public/MLDeformerModelRegistry.h
@@ -0,0 +1,36 @@
+// Copyright Epic Games, Inc. All Rights Reserved.
+
+#pragma once
+
+#include "CoreMinimal.h"
+
+class UMLDeformerModel;
+
+namespace UE::MLDeformer
+{
+	class FMLDeformerEditorModel;
+
+	DECLARE_DELEGATE_RetVal(FMLDeformerEditorModel*, FOnGetEditorModelInstance);
+
+	class MLDEFORMERFRAMEWORKEDITOR_API FMLDeformerEditorModelRegistry
+	{
+	public:
+		~FMLDeformerEditorModelRegistry();
+
+		void RegisterEditorModel(UClass* ModelType, FOnGetEditorModelInstance Delegate);
+		void UnregisterEditorModel(const UClass* ModelType);
+		void RemoveEditorModelInstance(FMLDeformerEditorModel* EditorModel);
+
+		FMLDeformerEditorModel* CreateEditorModel(UMLDeformerModel* Model);
+		FMLDeformerEditorModel* GetEditorModel(UMLDeformerModel* Model);
+
+		int GetNumRegisteredModels() const { return Map.Num(); }
+		int GetNumInstancedModels() const { return InstanceMap.Num(); }
+		const TMap<UClass*, FOnGetEditorModelInstance>& GetRegisteredModels() const { return Map; }
+		const TMap<UMLDeformerModel*, FMLDeformerEditorModel*>& GetModelInstances() const { return InstanceMap; }
+
+	private:
+		TMap<UClass*, FOnGetEditorModelInstance> Map;
+		TMap<UMLDeformerModel*, FMLDeformerEditorModel*> InstanceMap;
+	};
+}	// namespace UE::MLDeformer
diff --git a/Engine/Plugins/Experimental/MLDeformer/MLDeformerFramework/Source/MLDeformerFrameworkEditor/Public/MLDeformerSampler.h b/Engine/Plugins/Experimental/MLDeformer/MLDeformerFramework/Source/MLDeformerFrameworkEditor/Public/MLDeformerSampler.h
new file mode 100644
--- /dev/null
+++ b/Engine/Plugins/Experimental/MLDeformer/MLDeformerFramework/Source/MLDeformerFrameworkEditor/Public/MLDeformerSampler.h
@@ -0,0 +1,122 @@
+// Copyright Epic Games, Inc. All Rights Reserved.
+
+#pragma once
+
+#include "CoreMinimal.h"
+#include "MLDeformerModel.h"
+#include "UObject/ObjectPtr.h"
+
+class UVertexDeltaModel;
+class UDebugSkelMeshComponent;
+class UGeometryCacheComponent;
+class UWorld;
+class FSkeletalMeshLODRenderData;
+class FSkinWeightVertexBuffer;
+class AActor;
+
+namespace UE::MLDeformer
+{
+	class FMLDeformerEditorModel;
+
+	enum class EVertexDeltaSpace
+	{
+		PreSkinning,	/** Apply the deltas before skinning. */
+		PostSkinning	/** Apply the deltas after skinning. */
+	};
+
+	/**
+	 * The input data sampler.
+	 * This class can sample bone rotations, curve values and vertex deltas.
+	 * It does this by creating two temp actors, one with skeletal mesh component and one with geom cache component.
+	 */
+	class MLDEFORMERFRAMEWORKEDITOR_API FMLDeformerSampler
+	{
+	public:
+		virtual ~FMLDeformerSampler();
+
+		/** Call Init whenever assets or other relevant settings change. */
+		virtual void Init(FMLDeformerEditorModel* Model);
+
+		/** Call this every time the frame changes. This will update all buffer contents. */
+		virtual void Sample(int32 AnimFrameIndex);
+
+		/** Get the time in seconds, at a given frame index. */
+		virtual float GetTimeAtFrame(int32 InAnimFrameIndex) const { return 0.0f; }
+
+		const TArray<FVector3f>& GetSkinnedVertexPositions() const { return SkinnedVertexPositions; }
+		const TArray<float>& GetVertexDeltas() const { return VertexDeltas; }
+		const TArray<float>& GetBoneRotations() const { return BoneRotations; }
+		const TArray<float>& GetCurveValues() const { return CurveValues; }
+		int32 GetNumImportedVertices() const { return NumImportedVertices; }
+		int32 GetNumBones() const;
+		EVertexDeltaSpace GetVertexDeltaSpace() const { return VertexDeltaSpace; }
+		SIZE_T CalcMemUsagePerFrameInBytes() const;
+		void SetVertexDeltaSpace(EVertexDeltaSpace DeltaSpace) { VertexDeltaSpace = DeltaSpace; }
+		bool IsInitialized() const { return SkelMeshActor.Get() != nullptr; }
+
+	protected:
+		virtual void CreateActors();
+		virtual void RegisterTargetComponents() {}
+		void ExtractSkinnedPositions(int32 LODIndex, TArray<FMatrix44f>& InBoneMatrices, TArray<FVector3f>& TempPositions, TArray<FVector3f>& OutPositions) const;
+		FMatrix44f CalcInverseSkinningTransform(int32 VertexIndex, const FSkeletalMeshLODRenderData& SkelMeshLODData, const FSkinWeightVertexBuffer& SkinWeightBuffer) const;
+		AActor* CreateNewActor(UWorld* InWorld, const FName& Name) const;
+
+		void UpdateSkeletalMeshComponent();
+		void UpdateSkinnedPositions();
+		void UpdateBoneRotations();
+		void UpdateCurveValues();
+
+	protected:
+		/** The vertex delta editor model used to sample. */
+		FMLDeformerEditorModel* EditorModel = nullptr;
+
+		/** The skeletal mesh actor used to sample the skinned vertex positions. */
+		TObjectPtr<AActor> SkelMeshActor = nullptr;
+
+		/** The actor used for the target mesh. */
+		TObjectPtr<AActor> TargetMeshActor = nullptr;
+
+		/** The vertex delta model associated with this sampler. */
+		TObjectPtr<UMLDeformerModel> Model = nullptr;
+
+		/** The skeletal mesh component used to sample skinned positions. */
+		TObjectPtr<UDebugSkelMeshComponent>	SkeletalMeshComponent = nullptr;
+
+		/** The skinned vertex positions. */
+		TArray<FVector3f> SkinnedVertexPositions;
+
+		/** A temp array to store vertex positions in. */
+		TArray<FVector3f> TempVertexPositions;
+
+		/** The sampled bone matrices. */
+		TArray<FMatrix44f> BoneMatrices;
+
+		/**
+		 * The vertex deltas as float buffer. The number of floats equals: NumImportedVerts * 3.
+		 * The layout is (xyz)(xyz)(xyz)(...)
+		 */
+		TArray<float> VertexDeltas;	
+
+		/**
+		 * The bone rotation floats.
+		 * The number of floats in the buffer equals to NumBones * 6.
+		 * The six floats represent two columns of the 3x3 rotation matrix of the bone.
+		 */
+		TArray<float> BoneRotations;
+
+		/** A float for each Skeleton animation curve. */
+		TArray<float> CurveValues;	
+
+		/** The current sample time, in seconds. */
+		float SampleTime = 0.0f;
+
+		/** The number of imported vertices of the skeletal mesh and geometry cache. This will be 8 for a cube. */
+		int32 NumImportedVertices = 0;
+
+		/** The animation frame we sampled the deltas for. */
+		int32 AnimFrameIndex = -1;
+
+		/** The vertex delta space (pre or post skinning) used when calculating the deltas. */
+		EVertexDeltaSpace VertexDeltaSpace = EVertexDeltaSpace::PreSkinning;
+	};
+}	// namespace UE::VertexDeltaModel
diff --git a/Engine/Plugins/Experimental/MLDeformer/MLDeformerFramework/Source/MLDeformerFrameworkEditor/Public/MLDeformerTrainingModel.h b/Engine/Plugins/Experimental/MLDeformer/MLDeformerFramework/Source/MLDeformerFrameworkEditor/Public/MLDeformerTrainingModel.h
new file mode 100644
--- /dev/null
+++ b/Engine/Plugins/Experimental/MLDeformer/MLDeformerFramework/Source/MLDeformerFrameworkEditor/Public/MLDeformerTrainingModel.h
@@ -0,0 +1,70 @@
+// Copyright Epic Games, Inc. All Rights Reserved.
+#pragma once
+
+#include "CoreMinimal.h"
+#include "UObject/Object.h"
+#include "UObject/ObjectPtr.h"
+#include "MLDeformerTrainingModel.generated.h"
+
+class UMLDeformerModel;
+
+namespace UE::MLDeformer
+{
+	class FMLDeformerEditorModel;
+}
+
+UCLASS(Blueprintable)
+class MLDEFORMERFRAMEWORKEDITOR_API UMLDeformerTrainingModel
+	: public UObject
+{
+	GENERATED_BODY()
+
+public:
+	void Init(UE::MLDeformer::FMLDeformerEditorModel* InEditorModel);
+
+	UE::MLDeformer::FMLDeformerEditorModel* GetEditorModel() const { return EditorModel; }
+	void SetEditorModel(UE::MLDeformer::FMLDeformerEditorModel* InModel) { EditorModel = InModel; }
+
+	UFUNCTION(BlueprintPure, Category = "Training Data")
+	UMLDeformerModel* GetModel() const;
+
+	/** Get the number of input transforms. */
+	UFUNCTION(BlueprintPure, Category = "Training Data")
+	int32 GetNumberSampleTransforms() const;
+
+	/** Get number of input curves. */
+	UFUNCTION(BlueprintPure, Category = "Training Data")
+	int32 GetNumberSampleCurves() const;
+
+	/** Get the number of output deltas. */
+	UFUNCTION(BlueprintPure, Category = "Training Data")
+	int32 GetNumberSampleDeltas() const;
+
+	/** Get the number of samples in this data set. */
+	UFUNCTION(BlueprintPure, Category = "Training Data")
+    int32 NumSamples() const;
+
+	/** Set the current sample frame. This will internally call the SampleFrame method, which will update the deltas, curve values and bone rotations. */
+	UFUNCTION(BlueprintCallable, Category = "Training Data")
+	bool SetCurrentSampleIndex(int32 Index);
+
+protected:
+	/** Sample a given frame. This updates the sample deltas, curves, and bone rotations. */
+    virtual bool SampleFrame(int32 Index);
+
+public:
+	// The delta values per vertex for this sample. This is updated after SetCurrentSampleIndex is called. Contains an xyz (3 floats) for each vertex.
+	UPROPERTY(EditAnywhere, BlueprintReadWrite, Category = "Training Data")
+	TArray<float> SampleDeltas;
+
+	// The curve weights. This is updated after SetCurrentSampleIndex is called.
+	UPROPERTY(EditAnywhere, BlueprintReadWrite, Category = "Training Data")
+	TArray<float> SampleCurveValues;
+
+	// The bone rotations in bone (local) space for this sample. This is updated after SetCurrentSampleIndex is called and is 6 floats per bone (2 columns of 3x3 rotation matrix).
+	UPROPERTY(EditAnywhere, BlueprintReadWrite, Category = "Training Data")
+	TArray<float> SampleBoneRotations;
+
+protected:
+	UE::MLDeformer::FMLDeformerEditorModel* EditorModel = nullptr;
+};
diff --git a/Engine/Plugins/Experimental/MLDeformer/MLDeformerFramework/Source/MLDeformerFrameworkEditor/Public/MLDeformerVizSettingsDetails.h b/Engine/Plugins/Experimental/MLDeformer/MLDeformerFramework/Source/MLDeformerFrameworkEditor/Public/MLDeformerVizSettingsDetails.h
new file mode 100644
--- /dev/null
+++ b/Engine/Plugins/Experimental/MLDeformer/MLDeformerFramework/Source/MLDeformerFrameworkEditor/Public/MLDeformerVizSettingsDetails.h
@@ -0,0 +1,57 @@
+// Copyright Epic Games, Inc. All Rights Reserved.
+
+#pragma once
+
+#include "CoreMinimal.h"
+#include "IDetailCustomization.h"
+#include "AssetRegistry/AssetData.h"
+#include "PropertyHandle.h"
+
+class IDetailCategoryBuilder;
+class IDetailLayoutBuilder;
+class USkeleton;
+class UMLDeformerModel;
+class UMLDeformerVizSettings;
+
+namespace UE::MLDeformer
+{
+	class FMLDeformerEditorModel;
+
+	class MLDEFORMERFRAMEWORKEDITOR_API FMLDeformerVizSettingsDetails
+		: public IDetailCustomization
+	{
+	public:
+		/** Makes a new instance of this detail layout class for a specific detail view requesting it. */
+		static TSharedRef<IDetailCustomization> MakeInstance();
+
+		// ILayoutDetails overrides.
+		virtual void CustomizeDetails(IDetailLayoutBuilder& DetailBuilder) override;
+		// ~END ILayoutDetails overrides.
+
+		IDetailLayoutBuilder* GetDetailLayoutBuilder() const { return DetailLayoutBuilder; }
+
+	protected:
+		bool FilterAnimSequences(const FAssetData& AssetData, USkeleton* Skeleton);
+		void OnResetToDefaultDeformerGraph(TSharedPtr<IPropertyHandle> PropertyHandle);
+		bool IsResetToDefaultDeformerGraphVisible(TSharedPtr<IPropertyHandle> PropertyHandle);
+
+		virtual bool UpdateMemberPointers(const TArray<TWeakObjectPtr<UObject>>& Objects);
+		virtual void CreateCategories();
+		virtual void AddTestSequenceErrors() {}
+		virtual void AddDeformerGraphErrors() {}
+		virtual void AddGroundTruth() {}
+
+	protected:
+		/** Associated detail layout builder. */
+		IDetailLayoutBuilder* DetailLayoutBuilder = nullptr;
+
+		TObjectPtr<UMLDeformerModel> Model = nullptr;
+		TObjectPtr<UMLDeformerVizSettings> VizSettings = nullptr;
+		FMLDeformerEditorModel* EditorModel = nullptr;
+
+		IDetailCategoryBuilder* SharedCategoryBuilder = nullptr;
+		IDetailCategoryBuilder* TestAssetsCategory = nullptr;
+		IDetailCategoryBuilder* LiveSettingsCategory = nullptr;
+		IDetailCategoryBuilder* TrainingMeshesCategoryBuilder = nullptr;
+	};
+}	// namespace UE::VertexDeltaModel
diff --git a/Engine/Plugins/Experimental/MLDeformer/MLDeformerFramework/Source/MLDeformerFrameworkEditor/Public/SMLDeformerCurvePickerWidget.h b/Engine/Plugins/Experimental/MLDeformer/MLDeformerFramework/Source/MLDeformerFrameworkEditor/Public/SMLDeformerCurvePickerWidget.h
new file mode 100644
--- /dev/null
+++ b/Engine/Plugins/Experimental/MLDeformer/MLDeformerFramework/Source/MLDeformerFrameworkEditor/Public/SMLDeformerCurvePickerWidget.h
@@ -0,0 +1,100 @@
+// Copyright Epic Games, Inc. All Rights Reserved.
+
+#pragma once
+
+#include "CoreMinimal.h"
+#include "UObject/Object.h"
+#include "SlateFwd.h"
+#include "Widgets/SCompoundWidget.h"
+#include "Widgets/DeclarativeSyntaxSupport.h"
+#include "Widgets/SCompoundWidget.h"
+#include "Widgets/Views/STableRow.h"
+
+class SComboButton;
+class SSearchBox;
+class USkeleton;
+
+DECLARE_DELEGATE_OneParam(FOnCurveSelectionChanged, const FString&);
+
+DECLARE_DELEGATE_RetVal(FString, FOnGetSelectedCurve);
+DECLARE_DELEGATE_RetVal(USkeleton*, FOnGetSkeleton);
+
+DECLARE_DELEGATE_OneParam(FOnCurveNamePicked, const FString&)
+
+namespace UE::MLDeformer
+{
+	class MLDEFORMERFRAMEWORKEDITOR_API SMLDeformerCurvePickerWidget
+		: public SCompoundWidget
+	{
+	public:
+		SLATE_BEGIN_ARGS(SMLDeformerCurvePickerWidget) {}
+		SLATE_EVENT(FOnCurveNamePicked, OnCurveNamePicked)
+		SLATE_EVENT(FOnGetSkeleton, OnGetSkeleton)
+		SLATE_END_ARGS()
+
+		void Construct(const FArguments& InArgs);
+
+		TSharedPtr<SSearchBox> GetFilterTextWidget() const { return SearchBox; }
+
+	private:
+		void RefreshListItems();
+		void FilterAvailableCurves();
+
+		void HandleSelectionChanged(TSharedPtr<FString> InItem, ESelectInfo::Type InSelectionType);
+		TSharedRef<ITableRow> HandleGenerateRow(TSharedPtr<FString> InItem, const TSharedRef<STableViewBase>& InOwnerTable);
+		void HandleFilterTextChanged(const FText& InFilterText);
+
+	private:
+		/** Delegate fired when a curve name is picked. */
+		FOnCurveNamePicked OnCurveNamePicked;
+
+		/** Provide us with a skeleton. */
+		FOnGetSkeleton OnGetSkeleton;
+
+		/** The search filter box. */
+		TSharedPtr<SSearchBox> SearchBox;
+
+		/** The skeleton to get the curves from. */
+		TObjectPtr<USkeleton> Skeleton;
+
+		/** The names of the curves we are displaying. */
+		TArray<TSharedPtr<FString>> CurveNames;
+
+		/** All the unique curve names we can find. */
+		TSet<FString> UniqueCurveNames;
+
+		/** The string we use to filter curve names. */
+		FString FilterText;
+
+		/** The list view used to display names. */
+		TSharedPtr<SListView<TSharedPtr<FString>>> NameListView;
+	};
+
+	class SCurveSelectionWidget
+		: public SCompoundWidget
+	{
+	public: 
+		SLATE_BEGIN_ARGS(SCurveSelectionWidget) {}
+			SLATE_EVENT(FOnCurveSelectionChanged, OnCurveSelectionChanged);
+			SLATE_EVENT(FOnGetSelectedCurve, OnGetSelectedCurve);
+			SLATE_EVENT(FOnGetSkeleton, OnGetSkeleton)
+		SLATE_END_ARGS();
+
+		void Construct(const FArguments& InArgs);
+
+	private: 
+		TSharedRef<SWidget> CreateSkeletonWidgetMenu();
+		void OnSelectionChanged(const FString& CurveName);
+		USkeleton* GetSkeleton() const;
+		FText GetCurrentCurveName() const;
+		FText GetFinalToolTip() const;
+
+	private:
+		TSharedPtr<SComboButton> CurvePickerButton;
+		FOnCurveSelectionChanged OnCurveSelectionChanged;
+		FOnGetSelectedCurve OnGetSelectedCurve;
+		FOnGetSkeleton OnGetSkeleton;
+		FText SuppliedToolTip;
+		FText SelectedCurve;
+	};
+}	// namespace UE::MLDeformer
diff --git a/Engine/Plugins/Experimental/MLDeformer/VertexDeltaModel/Shaders/Private/VertexDeltaModelDataInterface.ush b/Engine/Plugins/Experimental/MLDeformer/VertexDeltaModel/Shaders/Private/VertexDeltaModelDataInterface.ush
new file mode 100644
--- /dev/null
+++ b/Engine/Plugins/Experimental/MLDeformer/VertexDeltaModel/Shaders/Private/VertexDeltaModelDataInterface.ush
@@ -0,0 +1,25 @@
+// Copyright Epic Games, Inc. All Rights Reserved.
+
+#include "/Plugin/ComputeFramework/Private/ComputeKernelCommon.ush"
+
+uint DI_LOCAL(NumVertices);
+uint DI_LOCAL(InputStreamStart);
+float3 DI_LOCAL(VertexDeltaScale);
+float3 DI_LOCAL(VertexDeltaMean);
+float DI_LOCAL(Weight);
+
+Buffer<float> DI_LOCAL(PositionDeltaBuffer);
+Buffer<uint> DI_LOCAL(VertexMapBuffer);
+
+DI_IMPL_READ(ReadNumVertices, uint, void)
+{
+	return DI_LOCAL(NumVertices);
+}
+
+DI_IMPL_READ(ReadPositionDelta, float3, uint VertexIndex)
+{
+	uint BufferIndex = DI_LOCAL(VertexMapBuffer)[VertexIndex + DI_LOCAL(InputStreamStart)] * 3;
+	float3 Delta = float3(DI_LOCAL(PositionDeltaBuffer)[BufferIndex], DI_LOCAL(PositionDeltaBuffer)[BufferIndex + 1], DI_LOCAL(PositionDeltaBuffer)[BufferIndex + 2]);
+	float3 FinalDelta = Delta * DI_LOCAL(VertexDeltaScale) + DI_LOCAL(VertexDeltaMean);
+	return FinalDelta * DI_LOCAL(Weight);
+}
diff --git a/Engine/Plugins/Experimental/MLDeformer/VertexDeltaModel/Source/VertexDeltaModel/Private/VertexDeltaModel.cpp b/Engine/Plugins/Experimental/MLDeformer/VertexDeltaModel/Source/VertexDeltaModel/Private/VertexDeltaModel.cpp
new file mode 100644
--- /dev/null
+++ b/Engine/Plugins/Experimental/MLDeformer/VertexDeltaModel/Source/VertexDeltaModel/Private/VertexDeltaModel.cpp
@@ -0,0 +1,56 @@
+// Copyright Epic Games, Inc. All Rights Reserved.
+#pragma once
+
+#include "VertexDeltaModel.h"
+#include "VertexDeltaModelVizSettings.h"
+#include "MLDeformerGeomCacheHelpers.h"
+#include "UObject/Object.h"
+
+#define LOCTEXT_NAMESPACE "UVertexDeltaModel"
+
+UVertexDeltaModel::UVertexDeltaModel(const FObjectInitializer& ObjectInitializer)
+	: Super(ObjectInitializer)
+{
+#if WITH_EDITORONLY_DATA
+	VizSettings = ObjectInitializer.CreateEditorOnlyDefaultSubobject<UVertexDeltaModelVizSettings>(this, TEXT("VizSettings"));
+#endif
+}
+
+#if WITH_EDITOR
+	void UVertexDeltaModel::UpdateNumTargetMeshVertices()
+	{
+		NumTargetMeshVerts = UE::MLDeformer::ExtractNumImportedGeomCacheVertices(GeometryCache);
+	}
+#endif // WITH_EDITOR
+
+#if WITH_EDITORONLY_DATA
+	void UVertexDeltaModel::SampleGroundTruthPositions(float SampleTime, TArray<FVector3f>& OutPositions)
+	{
+		UVertexDeltaModelVizSettings* VertexVizSettings = Cast<UVertexDeltaModelVizSettings>(VizSettings);
+		check(VertexVizSettings);
+
+		UGeometryCache* GeomCache = VertexVizSettings->GetTestGroundTruth();
+		if (GeomCache == nullptr)
+		{
+			OutPositions.Reset();
+			return;
+		}
+
+		if (MeshMappings.IsEmpty())
+		{
+			TArray<FString> FailedImportedMeshnames;
+			GenerateGeomCacheMeshMappings(SkeletalMesh, GeomCache, MeshMappings, FailedImportedMeshnames);
+		}
+
+		UE::MLDeformer::SampleGeomCachePositions(
+			0,
+			SampleTime,
+			MeshMappings,
+			SkeletalMesh,
+			GeomCache,
+			AlignmentTransform,
+			OutPositions);
+	}
+#endif
+
+#undef LOCTEXT_NAMESPACE
diff --git a/Engine/Plugins/Experimental/MLDeformer/VertexDeltaModel/Source/VertexDeltaModel/Private/VertexDeltaModelDataInterface.cpp b/Engine/Plugins/Experimental/MLDeformer/VertexDeltaModel/Source/VertexDeltaModel/Private/VertexDeltaModelDataInterface.cpp
new file mode 100644
--- /dev/null
+++ b/Engine/Plugins/Experimental/MLDeformer/VertexDeltaModel/Source/VertexDeltaModel/Private/VertexDeltaModelDataInterface.cpp
@@ -0,0 +1,53 @@
+// Copyright Epic Games, Inc. All Rights Reserved.
+
+#include "VertexDeltaModelDataInterface.h"
+#include "VertexDeltaModel.h"
+#include "MLDeformerModelInstance.h"
+#include "MLDeformerAsset.h"
+#include "MLDeformerComponent.h"
+#include "Components/SkeletalMeshComponent.h"
+#include "ComputeFramework/ComputeKernelPermutationSet.h"
+#include "ComputeFramework/ShaderParamTypeDefinition.h"
+#include "ComputeFramework/ShaderParameterMetadataBuilder.h"
+#include "OptimusDataDomain.h"
+#include "RenderGraphBuilder.h"
+#include "RenderGraphUtils.h"
+#include "RenderGraphResources.h"
+#include "SkeletalRenderPublic.h"
+
+BEGIN_SHADER_PARAMETER_STRUCT(FVertexDeltaModelDataInterfaceParameters,)
+	MLDEFORMER_SHADER_PARAMETERS()
+	SHADER_PARAMETER(FVector3f, VertexDeltaScale)
+	SHADER_PARAMETER(FVector3f, VertexDeltaMean)
+END_SHADER_PARAMETER_STRUCT()
+
+MLDEFORMER_GRAPH_IMPLEMENT_BASICS(
+	UVertexDeltaModelDataInterface,
+	UVertexDeltaModelDataProvider,
+	UE::VertexDeltaModel::FVertexDeltaModelDataProviderProxy,
+	FVertexDeltaModelDataInterfaceParameters,
+	TEXT("#include \"/Plugin/VertexDeltaModel/Private/VertexDeltaModelDataInterface.ush\"\n"),
+	TEXT("Vertex Delta Model"))
+
+namespace UE::VertexDeltaModel
+{
+	FVertexDeltaModelDataProviderProxy::FVertexDeltaModelDataProviderProxy(USkeletalMeshComponent* SkeletalMeshComponent, UMLDeformerComponent* DeformerComponent)
+		: UE::MLDeformer::FMLDeformerGraphDataProviderProxy(SkeletalMeshComponent, DeformerComponent)
+	{
+		const UMLDeformerModel* Model = DeformerComponent->GetDeformerAsset()->GetModel();
+		const UVertexDeltaModel* VertexDeltaModel = Cast<UVertexDeltaModel>(Model);
+		check(VertexDeltaModel);
+
+		VertexDeltaScale = static_cast<FVector3f>(VertexDeltaModel->GetVertexDeltaScale());
+		VertexDeltaMean = static_cast<FVector3f>(VertexDeltaModel->GetVertexDeltaMean());
+	}
+
+	void FVertexDeltaModelDataProviderProxy::GatherDispatchData(FDispatchSetup const& InDispatchSetup, FCollectedDispatchData& InOutDispatchData)
+	{
+		MLDEFORMER_GRAPH_DISPATCH_START(FVertexDeltaModelDataInterfaceParameters, InDispatchSetup, InOutDispatchData)
+		MLDEFORMER_GRAPH_DISPATCH_DEFAULT_PARAMETERS()
+		Parameters->VertexDeltaScale = VertexDeltaScale;
+		Parameters->VertexDeltaMean = VertexDeltaMean;
+		MLDEFORMER_GRAPH_DISPATCH_END()
+	}
+}	// namespace UE::VertexDeltaModel
diff --git a/Engine/Plugins/Experimental/MLDeformer/VertexDeltaModel/Source/VertexDeltaModel/Private/VertexDeltaModelModule.cpp b/Engine/Plugins/Experimental/MLDeformer/VertexDeltaModel/Source/VertexDeltaModel/Private/VertexDeltaModelModule.cpp
new file mode 100644
--- /dev/null
+++ b/Engine/Plugins/Experimental/MLDeformer/VertexDeltaModel/Source/VertexDeltaModel/Private/VertexDeltaModelModule.cpp
@@ -0,0 +1,25 @@
+// Copyright Epic Games, Inc. All Rights Reserved.
+
+#include "VertexDeltaModelModule.h"
+#include "Modules/ModuleManager.h"
+#include "Interfaces/IPluginManager.h"
+#include "Misc/Paths.h"
+#include "ShaderCore.h"
+
+VERTEXDELTAMODEL_API DEFINE_LOG_CATEGORY(LogVertexDeltaModel)
+
+#define LOCTEXT_NAMESPACE "VertexDeltaModelModule"
+
+IMPLEMENT_MODULE(UE::VertexDeltaModel::FVertexDeltaModelModule, VertexDeltaModel)
+
+namespace UE::VertexDeltaModel
+{
+	void FVertexDeltaModelModule::StartupModule()
+	{
+		// Register an additional shader path for our shaders used inside the deformer graph system.
+		const FString PluginShaderDir = FPaths::Combine(IPluginManager::Get().FindPlugin(TEXT("VertexDeltaModel"))->GetBaseDir(), TEXT("Shaders"));
+		AddShaderSourceDirectoryMapping(TEXT("/Plugin/VertexDeltaModel"), PluginShaderDir);
+	}
+}	// namespace UE::VertexDeltaModel
+
+#undef LOCTEXT_NAMESPACE
diff --git a/Engine/Plugins/Experimental/MLDeformer/VertexDeltaModel/Source/VertexDeltaModel/Public/VertexDeltaModel.h b/Engine/Plugins/Experimental/MLDeformer/VertexDeltaModel/Source/VertexDeltaModel/Public/VertexDeltaModel.h
new file mode 100644
--- /dev/null
+++ b/Engine/Plugins/Experimental/MLDeformer/VertexDeltaModel/Source/VertexDeltaModel/Public/VertexDeltaModel.h
@@ -0,0 +1,129 @@
+// Copyright Epic Games, Inc. All Rights Reserved.
+#pragma once
+
+#include "CoreMinimal.h"
+#include "MLDeformerModel.h"
+#include "MLDeformerVizSettings.h"
+#include "MLDeformerInputInfo.h"
+#include "MLDeformerGeomCacheHelpers.h"
+#include "UObject/Object.h"
+#include "VertexDeltaModel.generated.h"
+
+class USkeletalMesh;
+class UGeometryCache;
+class UAnimSequence;
+class UNeuralNetwork;
+class UMLDeformerAsset;
+class USkeleton;
+class IPropertyHandle;
+
+/** The activation function to use during the ML Deformer training process. */
+UENUM()
+enum class EVertexDeltaModelActivationFunction : uint8
+{
+	Relu,
+	LRelu,
+	Tanh
+};
+
+/** The loss function to use during the ML Deformer training process. */
+UENUM()
+enum class EVertexDeltaModelLossFunction : uint8
+{
+	L1,
+	MSE,
+	Shrinkage
+};
+
+UCLASS()
+class VERTEXDELTAMODEL_API UVertexDeltaModel 
+	: public UMLDeformerModel
+{
+	GENERATED_BODY()
+
+public:
+	UVertexDeltaModel(const FObjectInitializer& ObjectInitializer);
+
+	// UMLDeformerModel overrides.
+	virtual FString GetDisplayName() const override { return "Vertex Delta Model"; }
+#if WITH_EDITORONLY_DATA
+	virtual bool HasTrainingGroundTruth() const override { return (GeometryCache.Get() != nullptr); }
+	virtual void SampleGroundTruthPositions(float SampleTime, TArray<FVector3f>& OutPositions) override;
+#endif
+#if WITH_EDITOR
+	virtual void UpdateNumTargetMeshVertices();
+#endif
+	// ~END UMLDeformerModel overrides.
+
+	const FVector& GetVertexDeltaMean() const { return VertexDeltaMean; }
+	const FVector& GetVertexDeltaScale() const { return VertexDeltaScale; }
+
+#if WITH_EDITORONLY_DATA
+	const UGeometryCache* GetGeometryCache() const { return GeometryCache; }
+	UGeometryCache* GetGeometryCache() { return GeometryCache; }
+	EVertexDeltaModelLossFunction GetLossFunction() const { return LossFunction; }
+	EVertexDeltaModelActivationFunction GetActivationFunction() const { return ActivationFunction; }
+	int32 GetNumHiddenLayers() const { return NumHiddenLayers; }
+	int32 GetNumNeuronsPerLayer() const { return NumNeuronsPerLayer; }
+	int32 GetNumEpochs() const { return Epochs; }
+	int32 GetBatchSize() const { return BatchSize; }
+	float GetLearningRate() const { return LearningRate; }
+	float GetShrinkageSpeed() const { return ShrinkageSpeed; }
+	float GetShrinkageThrehsold() const { return ShrinkageThreshold; }
+	TArray<UE::MLDeformer::FMLDeformerGeomCacheMeshMapping>& GetGeomCacheMeshMappings() { return MeshMappings; }
+	const TArray<UE::MLDeformer::FMLDeformerGeomCacheMeshMapping>& GetGeomCacheMeshMappings() const { return MeshMappings; }
+#endif
+
+public:
+	/** The vertex delta mean. This is passed to the deformer shader to reconstruct the correct output deltas. */
+	UPROPERTY()
+	FVector VertexDeltaMean = FVector::ZeroVector;
+
+	/** The vertex delta scale. This is passed to the deformer shader to rescale the output deltas. */
+	UPROPERTY()
+	FVector VertexDeltaScale = FVector::OneVector;
+
+#if WITH_EDITORONLY_DATA
+	TArray<UE::MLDeformer::FMLDeformerGeomCacheMeshMapping> MeshMappings;
+
+	/** The geometry cache that represents the complex mesh deformations. */
+	UPROPERTY(EditAnywhere, Category = "Target Mesh")
+	TObjectPtr<UGeometryCache> GeometryCache = nullptr;
+
+	/** The number of hidden layers that the neural network model will have.\nHigher numbers will slow down performance but can deal with more complex deformations. */
+	UPROPERTY(EditAnywhere, BlueprintReadWrite, Category = "Training Settings", meta = (ClampMin = "1", ClampMax = "10"))
+	int32 NumHiddenLayers = 2;
+
+	/** The number of units/neurons per hidden layer. Higher numbers will slow down performance but allow for more complex mesh deformations. */
+	UPROPERTY(EditAnywhere, BlueprintReadWrite, Category = "Training Settings", meta = (ClampMin = "1"))
+	int32 NumNeuronsPerLayer = 256;
+
+	/** The number of epochs to process without any decay. */
+	UPROPERTY(EditAnywhere, BlueprintReadWrite, Category = "Training Settings", meta = (ClampMin = "1"))
+	int32 Epochs = 20;
+
+	/** The number of frames per batch when training the model. */
+	UPROPERTY(EditAnywhere, BlueprintReadWrite, AdvancedDisplay, Category = "Training Settings", meta = (ClampMin = "1"))
+	int32 BatchSize = 128;
+
+	/** The learning rate used during the model training. */
+	UPROPERTY(EditAnywhere, BlueprintReadWrite, AdvancedDisplay, Category = "Training Settings", meta = (ClampMin = "0.000001", ClampMax = "1.0"))
+	float LearningRate = 0.00175f;
+
+	/** The activation function to use in the neural network. */
+	UPROPERTY(EditAnywhere, BlueprintReadWrite, AdvancedDisplay, Category = "Training Settings")
+	EVertexDeltaModelActivationFunction ActivationFunction = EVertexDeltaModelActivationFunction::LRelu;
+
+	/** The loss function to use during model training. */
+	UPROPERTY(EditAnywhere, BlueprintReadWrite, AdvancedDisplay, Category = "Training Settings")
+	EVertexDeltaModelLossFunction LossFunction = EVertexDeltaModelLossFunction::L1;
+
+	/** Shrinkage speed. Only if the shrinkage loss is used. */
+	UPROPERTY(EditAnywhere, BlueprintReadWrite, AdvancedDisplay, Category = "Training Settings", meta = (ClampMin = "1.0", ClampMax = "1000.0"))
+	float ShrinkageSpeed = 10.0f;
+
+	/** Shrinkage threshold. Only if the shrinkage loss is used. */
+	UPROPERTY(EditAnywhere, BlueprintReadWrite, AdvancedDisplay, Category = "Training Settings", meta = (ClampMin = "0.0"))
+	float ShrinkageThreshold = 0.1f;
+#endif // WITH_EDITORONLY_DATA
+};
diff --git a/Engine/Plugins/Experimental/MLDeformer/VertexDeltaModel/Source/VertexDeltaModel/Public/VertexDeltaModelDataInterface.h b/Engine/Plugins/Experimental/MLDeformer/VertexDeltaModel/Source/VertexDeltaModel/Public/VertexDeltaModelDataInterface.h
new file mode 100644
--- /dev/null
+++ b/Engine/Plugins/Experimental/MLDeformer/VertexDeltaModel/Source/VertexDeltaModel/Public/VertexDeltaModelDataInterface.h
@@ -0,0 +1,69 @@
+// Copyright Epic Games, Inc. All Rights Reserved.
+
+#pragma once
+
+#include "CoreMinimal.h"
+#include "MLDeformerGraphDataInterface.h"
+#include "OptimusComputeDataInterface.h"
+#include "ComputeFramework/ComputeDataProvider.h"
+#include "VertexDeltaModelDataInterface.generated.h"
+
+class FRDGBuffer;
+class FRDGBufferSRV;
+class FRHIShaderResourceView;
+class FSkeletalMeshObject;
+class UMLDeformerComponent;
+class UNeuralNetwork;
+class USkeletalMeshComponent;
+class UMLDeformerModel;
+
+/** Compute Framework Data Interface for MLDefomer data. */
+UCLASS(Category = ComputeFramework)
+class VERTEXDELTAMODEL_API UVertexDeltaModelDataInterface
+	: public UMLDeformerGraphDataInterface
+{
+	GENERATED_BODY()
+
+public:
+	// UOptimusComputeDataInterface overrides.
+	virtual FString GetDisplayName() const override;
+	// ~END UOptimusComputeDataInterface overrides.
+
+	// UComputeDataInterface overrides.
+	virtual void GetShaderParameters(TCHAR const* UID, FShaderParametersMetadataBuilder& OutBuilder) const override;
+	virtual void GetHLSL(FString& OutHLSL) const override;
+	virtual UComputeDataProvider* CreateDataProvider(TArrayView<TObjectPtr<UObject>> InSourceObjects, uint64 InInputMask, uint64 InOutputMask) const override;
+	// ~END UComputeDataInterface overrides.
+};
+
+/** Compute Framework Data Provider for MLDeformer data. */
+UCLASS(BlueprintType, EditInlineNew, Category = ComputeFramework)
+class VERTEXDELTAMODEL_API UVertexDeltaModelDataProvider
+	: public UMLDeformerGraphDataProvider
+{
+	GENERATED_BODY()
+
+public:
+	// UComputeDataProvider overrides.
+	virtual FComputeDataProviderRenderProxy* GetRenderProxy() override;
+	// ~END UComputeDataProvider overrides.
+};
+
+namespace UE::VertexDeltaModel
+{
+	/** Compute Framework Data Provider Proxy for MLDeformer data. */
+	class VERTEXDELTAMODEL_API FVertexDeltaModelDataProviderProxy
+		: public UE::MLDeformer::FMLDeformerGraphDataProviderProxy
+	{
+	public:
+		FVertexDeltaModelDataProviderProxy(USkeletalMeshComponent* SkeletalMeshComponent, UMLDeformerComponent* DeformerComponent);
+
+		// FComputeDataProviderRenderProxy overrides.
+		virtual void GatherDispatchData(FDispatchSetup const& InDispatchSetup, FCollectedDispatchData& InOutDispatchData) override;
+		// ~END FComputeDataProviderRenderProxy overrides.
+
+	private:
+		FVector3f VertexDeltaScale = FVector3f(1.0f, 1.0f, 1.0f);
+		FVector3f VertexDeltaMean = FVector3f(0.0f, 0.0f, 0.0f);
+	};
+}	// namespace UE::VertexDeltaModel
diff --git a/Engine/Plugins/Experimental/MLDeformer/VertexDeltaModel/Source/VertexDeltaModel/Public/VertexDeltaModelModule.h b/Engine/Plugins/Experimental/MLDeformer/VertexDeltaModel/Source/VertexDeltaModel/Public/VertexDeltaModelModule.h
new file mode 100644
--- /dev/null
+++ b/Engine/Plugins/Experimental/MLDeformer/VertexDeltaModel/Source/VertexDeltaModel/Public/VertexDeltaModelModule.h
@@ -0,0 +1,20 @@
+// Copyright Epic Games, Inc. All Rights Reserved.
+
+#pragma once
+
+#include "CoreMinimal.h"
+#include "Modules/ModuleManager.h"
+
+VERTEXDELTAMODEL_API DECLARE_LOG_CATEGORY_EXTERN(LogVertexDeltaModel, Log, All);
+
+namespace UE::VertexDeltaModel
+{
+	class VERTEXDELTAMODEL_API FVertexDeltaModelModule
+		: public IModuleInterface
+	{
+	public:
+		// IModuleInterface overrides.
+		virtual void StartupModule() override;
+		// ~END IModuleInterface overrides.
+	};
+}	// namespace VertexDeltaModel
diff --git a/Engine/Plugins/Experimental/MLDeformer/VertexDeltaModel/Source/VertexDeltaModel/Public/VertexDeltaModelVizSettings.h b/Engine/Plugins/Experimental/MLDeformer/VertexDeltaModel/Source/VertexDeltaModel/Public/VertexDeltaModelVizSettings.h
new file mode 100644
--- /dev/null
+++ b/Engine/Plugins/Experimental/MLDeformer/VertexDeltaModel/Source/VertexDeltaModel/Public/VertexDeltaModelVizSettings.h
@@ -0,0 +1,30 @@
+// Copyright Epic Games, Inc. All Rights Reserved.
+
+#pragma once
+
+#include "CoreMinimal.h"
+#include "UObject/Object.h"
+#include "VertexDeltaModelVizSettings.generated.h"
+
+class UGeometryCache;
+
+/**
+ * The vizualization settings specific to the the vertex delta model.
+ */
+UCLASS()
+class VERTEXDELTAMODEL_API UVertexDeltaModelVizSettings 
+	: public UMLDeformerVizSettings
+{
+	GENERATED_BODY()
+
+public:
+#if WITH_EDITORONLY_DATA
+	virtual bool HasTestGroundTruth() const override { return (GroundTruth.Get() != nullptr); }
+	UGeometryCache* GetTestGroundTruth() const { return GroundTruth; }
+
+public:
+	/** The geometry cache that represents the ground truth of the test anim sequence. */
+	UPROPERTY(EditAnywhere, Category = "Test Assets")
+	TObjectPtr<UGeometryCache> GroundTruth = nullptr;
+#endif // WITH_EDITORONLY_DATA
+};
diff --git a/Engine/Plugins/Experimental/MLDeformer/VertexDeltaModel/Source/VertexDeltaModelEditor/Private/VertexDeltaEditorModel.cpp b/Engine/Plugins/Experimental/MLDeformer/VertexDeltaModel/Source/VertexDeltaModelEditor/Private/VertexDeltaEditorModel.cpp
new file mode 100644
--- /dev/null
+++ b/Engine/Plugins/Experimental/MLDeformer/VertexDeltaModel/Source/VertexDeltaModelEditor/Private/VertexDeltaEditorModel.cpp
@@ -0,0 +1,305 @@
+// Copyright Epic Games, Inc. All Rights Reserved.
+
+#include "VertexDeltaEditorModel.h"
+#include "VertexDeltaModel.h"
+#include "VertexDeltaModelVizSettings.h"
+#include "VertexDeltaTrainingModel.h"
+#include "VertexDeltaEditorModelActor.h"
+#include "MLDeformerEditorToolkit.h"
+#include "MLDeformerEditorStyle.h"
+#include "MLDeformerInputInfo.h"
+#include "MLDeformerGeomCacheSampler.h"
+#include "MLDeformerGeomCacheHelpers.h"
+#include "GeometryCache.h"
+#include "GeometryCacheComponent.h"
+#include "Animation/DebugSkelMeshComponent.h"
+#include "Rendering/SkeletalMeshModel.h"
+#include "Rendering/SkeletalMeshLODModel.h"
+#include "Engine/SkeletalMesh.h"
+
+#define LOCTEXT_NAMESPACE "VertexDeltaEditorModel"
+
+namespace UE::VertexDeltaModel
+{
+	using namespace UE::MLDeformer;
+
+	FVertexDeltaEditorModel::~FVertexDeltaEditorModel()
+	{
+		Model->OnPostEditChangeProperty().Unbind();
+	}
+
+	FMLDeformerEditorModel* FVertexDeltaEditorModel::MakeInstance()
+	{
+		return new FVertexDeltaEditorModel();
+	}
+
+	void FVertexDeltaEditorModel::AddReferencedObjects(FReferenceCollector& Collector)
+	{
+		Collector.AddReferencedObject(EditorInputInfo);
+	}
+
+	FMLDeformerSampler* FVertexDeltaEditorModel::CreateSampler() const 
+	{ 
+		FMLDeformerGeomCacheSampler* NewSampler = new FMLDeformerGeomCacheSampler();
+		NewSampler->OnGetGeometryCache().BindLambda([this] { return GetVertexDeltaModel()->GetGeometryCache(); });
+		return NewSampler;
+	}
+
+	void FVertexDeltaEditorModel::Init(const InitSettings& InitSettings)
+	{
+		FMLDeformerEditorModel::Init(InitSettings);
+		Model->OnPostEditChangeProperty().BindRaw(this, &FVertexDeltaEditorModel::OnPostEditChangeProperty);
+	}
+
+	void FVertexDeltaEditorModel::OnPostEditChangeProperty(FPropertyChangedEvent& PropertyChangedEvent)
+	{
+		const FProperty* Property = PropertyChangedEvent.Property;
+		if (Property == nullptr)
+		{
+			return;
+		}
+
+		HandleDefaultPropertyChanges(PropertyChangedEvent);
+
+		// When we change one of these properties below, restart animations etc.
+		if (Property->GetFName() == GET_MEMBER_NAME_CHECKED(UVertexDeltaModel, GeometryCache) || 
+			Property->GetFName() == GET_MEMBER_NAME_CHECKED(UVertexDeltaModelVizSettings, GroundTruth))
+		{
+			TriggerInputAssetChanged(true);
+		}
+		else
+		if (Property->GetFName() == GET_MEMBER_NAME_CHECKED(UVertexDeltaModel, DeltaCutoffLength) || 
+			Property->GetFName() == GET_MEMBER_NAME_CHECKED(UMLDeformerModel, AlignmentTransform))
+		{
+			if (PropertyChangedEvent.ChangeType == EPropertyChangeType::ValueSet)
+			{
+				bIsDataNormalized = false;
+				SampleDeltas();
+			}
+		}
+	}
+
+	void FVertexDeltaEditorModel::OnInputAssetsChanged()
+	{
+		// Update the skeletal mesh components of the training, test base, and ml deformed actor.
+		FMLDeformerEditorModel::OnInputAssetsChanged();
+
+		UVertexDeltaModelVizSettings* VizSettings = GetVertexDeltaModelVizSettings();
+		check(VizSettings);
+		const float TestAnimSpeed = VizSettings->GetAnimPlaySpeed();
+		UAnimSequence* TestAnimSequence = VizSettings->GetTestAnimSequence();
+
+		// Update the training geometry cache.
+		UGeometryCacheComponent* GeometryCacheComponent = FindVertexDeltaModelEditorActor(ActorID_Train_GroundTruth)->GetGeometryCacheComponent();
+		check(GeometryCacheComponent);
+		GeometryCacheComponent->SetGeometryCache(GetVertexDeltaModel()->GetGeometryCache());
+		GeometryCacheComponent->ResetAnimationTime();
+		GeometryCacheComponent->SetLooping(false);
+		GeometryCacheComponent->SetManualTick(true);
+		GeometryCacheComponent->SetPlaybackSpeed(TestAnimSpeed);
+		GeometryCacheComponent->Play();
+
+		// Update the test geometry cache (ground truth) component.
+		GeometryCacheComponent = FindVertexDeltaModelEditorActor(ActorID_Test_GroundTruth)->GetGeometryCacheComponent();
+		GeometryCacheComponent->SetGeometryCache(VizSettings->GetTestGroundTruth());
+		GeometryCacheComponent->ResetAnimationTime();
+		GeometryCacheComponent->SetLooping(true);
+		GeometryCacheComponent->SetManualTick(true);
+		GeometryCacheComponent->SetPlaybackSpeed(TestAnimSpeed);
+		GeometryCacheComponent->Play();
+
+		UVertexDeltaModel* VertexDeltaModel = GetVertexDeltaModel();
+		VertexDeltaModel->MeshMappings.Reset();
+	}
+
+	void FVertexDeltaEditorModel::CreateTrainingGroundTruthActor(UWorld* World)
+	{
+		UGeometryCache* GeomCache = GetVertexDeltaModel()->GetGeometryCache();
+		const FLinearColor LabelColor = FMLDeformerEditorStyle::Get().GetColor("MLDeformer.TargetMesh.LabelColor");
+		const FLinearColor WireframeColor = FMLDeformerEditorStyle::Get().GetColor("MLDeformer.TargetMesh.WireframeColor");
+		CreateGeomCacheActor(
+			World, 
+			ActorID_Train_GroundTruth, 
+			"Train GroundTruth", 
+			GeomCache, 
+			LabelColor, 
+			WireframeColor, 
+			LOCTEXT("TrainGroundTruthActorLabelText", "Target Mesh"),
+			true);
+	}
+
+	void FVertexDeltaEditorModel::CreateTestGroundTruthActor(UWorld* World)
+	{
+		UGeometryCache* GeomCache = GetVertexDeltaModelVizSettings()->GetTestGroundTruth();
+		const FLinearColor LabelColor = FMLDeformerEditorStyle::Get().GetColor("MLDeformer.GroundTruth.LabelColor");
+		const FLinearColor WireframeColor = FMLDeformerEditorStyle::Get().GetColor("MLDeformer.GroundTruth.WireframeColor");
+		CreateGeomCacheActor(
+			World, 
+			ActorID_Test_GroundTruth, 
+			"Test GroundTruth", 
+			GeomCache, 
+			LabelColor, 
+			WireframeColor,
+			LOCTEXT("TestGroundTruthActorLabelText", "Ground Truth"),
+			false);
+	}
+
+	void FVertexDeltaEditorModel::CreateGeomCacheActor(UWorld* World, int32 ActorID, const FName& Name, UGeometryCache* GeomCache, FLinearColor LabelColor, FLinearColor WireframeColor, const FText& LabelText, bool bIsTrainingActor)
+	{
+		FActorSpawnParameters SpawnParams;
+		SpawnParams.Name = MakeUniqueObjectName(World, AActor::StaticClass(), Name);
+		AActor* Actor = World->SpawnActor<AActor>(SpawnParams);
+		Actor->SetFlags(RF_Transient);
+
+		// Create the Geometry Cache Component.
+		UGeometryCacheComponent* GeomCacheComponent = NewObject<UGeometryCacheComponent>(Actor);
+		GeomCacheComponent->SetGeometryCache(GeomCache);
+		GeomCacheComponent->RegisterComponent();
+		GeomCacheComponent->SetOverrideWireframeColor(true);
+		GeomCacheComponent->SetWireframeOverrideColor(WireframeColor);
+		GeomCacheComponent->MarkRenderStateDirty();
+		GeomCacheComponent->SetVisibility(false);
+		Actor->SetRootComponent(GeomCacheComponent);
+
+		// Create the editor actor.
+		FMLDeformerEditorActor::FConstructSettings Settings;
+		Settings.Actor = Actor;
+		Settings.TypeID = ActorID;
+		Settings.LabelColor = LabelColor;
+		Settings.LabelText = LabelText;
+		Settings.bIsTrainingActor = bIsTrainingActor;
+		FVertexDeltaEditorModelActor* EditorActor = static_cast<FVertexDeltaEditorModelActor*>(CreateEditorActor(Settings));
+		EditorActor->SetGeometryCacheComponent(GeomCacheComponent);
+		EditorActors.Add(EditorActor);
+	} 
+
+	FMLDeformerEditorActor* FVertexDeltaEditorModel::CreateEditorActor(const FMLDeformerEditorActor::FConstructSettings& Settings) const
+	{
+		return new FVertexDeltaEditorModelActor(Settings);
+	}
+
+	FVertexDeltaEditorModelActor* FVertexDeltaEditorModel::FindVertexDeltaModelEditorActor(int32 TypeID) const
+	{
+		return static_cast<FVertexDeltaEditorModelActor*>(FindEditorActor(TypeID));
+	}
+
+	UVertexDeltaModelVizSettings* FVertexDeltaEditorModel::GetVertexDeltaModelVizSettings() const
+	{
+		return Cast<UVertexDeltaModelVizSettings>(GetVertexDeltaModel()->GetVizSettings());
+	}
+
+	double FVertexDeltaEditorModel::GetTimeAtFrame(int32 FrameNumber) const
+	{
+		const FVertexDeltaEditorModelActor* EditorActor = static_cast<FVertexDeltaEditorModelActor*>(GetTimelineEditorActor());
+		if (EditorActor && EditorActor->GetGeometryCacheComponent())
+		{
+			return EditorActor->GetGeometryCacheComponent()->GetTimeAtFrame(FrameNumber);
+		}
+		return 0.0;
+	}
+
+	int32 FVertexDeltaEditorModel::GetFrameAtTime(double TimeInSeconds) const
+	{
+		const FVertexDeltaEditorModelActor* EditorActor = static_cast<FVertexDeltaEditorModelActor*>(GetTimelineEditorActor());
+		if (EditorActor && EditorActor->GetGeometryCacheComponent())
+		{
+			return EditorActor->GetGeometryCacheComponent()->GetFrameAtTime(TimeInSeconds);
+		}
+		return 0;
+	}
+
+	int32 FVertexDeltaEditorModel::GetNumFrames() const
+	{
+		const UGeometryCache* GeometryCache = GetVertexDeltaModel()->GetGeometryCache();
+		if (GeometryCache == nullptr)
+		{
+			return 0;
+		}
+		const int32 StartFrame = GeometryCache->GetStartFrame();
+		const int32 EndFrame = GeometryCache->GetEndFrame();
+		return (EndFrame - StartFrame) + 1;
+	}
+
+	void FVertexDeltaEditorModel::UpdateIsReadyForTrainingState()
+	{
+		bIsReadyForTraining = false;
+
+		// Do some basic checks first, like if there is a skeletal mesh, ground truth, anim sequence, and if there are frames.
+		if (!IsEditorReadyForTrainingBasicChecks())
+		{
+			return;
+		}
+
+		// Now make sure the assets are compatible.
+		UVertexDeltaModel* VertexDeltaModel = GetVertexDeltaModel();
+		UGeometryCache* GeomCache = VertexDeltaModel->GetGeometryCache();
+		USkeletalMesh* SkeletalMesh = VertexDeltaModel->GetSkeletalMesh();
+		if (!GetGeomCacheVertexErrorText(SkeletalMesh, GeomCache, FText(), FText()).IsEmpty() ||
+			!GetGeomCacheErrorText(SkeletalMesh, GeomCache).IsEmpty())
+		{
+			return;
+		}
+
+		// Make sure every skeletal imported mesh has some geometry track.
+		const int32 NumGeomCacheTracks = GeomCache ? GeomCache->Tracks.Num() : 0;
+		int32 NumSkelMeshes = 0;
+		if (SkeletalMesh)
+		{
+			FSkeletalMeshModel* ImportedModel = SkeletalMesh->GetImportedModel();
+			if (ImportedModel)
+			{
+				NumSkelMeshes = ImportedModel->LODModels[0].ImportedMeshInfos.Num();		
+			}
+		}
+
+		// Allow the special case where there is just one mesh and track.
+		if (NumGeomCacheTracks != 1 || NumSkelMeshes != 1)
+		{
+			if (!GetGeomCacheSampler()->GetFailedImportedMeshNames().IsEmpty())
+			{
+				return;
+			}
+		}
+
+		bIsReadyForTraining = true;
+	}
+
+	ETrainingResult FVertexDeltaEditorModel::Train()
+	{
+		return TrainModel<UVertexDeltaTrainingModel>(this);
+	}
+
+	void FVertexDeltaEditorModel::OnPreTraining()
+	{
+		// Make a backup of the normalization values, as they get overwritten when training.
+		// However, when we abort, we want to restore to the original values again.
+		// See OnTrainingAborted for when we restore the backup again.
+		VertexDeltaMeanBackup = GetVertexDeltaModel()->GetVertexDeltaMean();
+		VertexDeltaScaleBackup = GetVertexDeltaModel()->GetVertexDeltaScale();
+	}
+
+	void FVertexDeltaEditorModel::OnTrainingAborted()
+	{
+		// Restore the vertex delta mean and scale, as we aborted, and they could have changed when training
+		// on a smaller subset of frames/samples. If we don't do this, the mesh will deform incorrectly.
+		GetVertexDeltaModel()->VertexDeltaMean = VertexDeltaMeanBackup;
+		GetVertexDeltaModel()->VertexDeltaScale = VertexDeltaScaleBackup;
+	}
+
+	FString FVertexDeltaEditorModel::GetTrainedNetworkOnnxFile() const
+	{
+		return FString(FPaths::ProjectIntermediateDir() + TEXT("VertexDeltaModel/latest_net_G.onnx"));
+	}
+
+	FString FVertexDeltaEditorModel::GetDefaultDeformerGraphAssetPath() const
+	{
+		return FString(TEXT("/VertexDeltaModel/Deformers/VertexDeltaModelGraph.VertexDeltaModelGraph"));
+	}
+
+	FString FVertexDeltaEditorModel::GetHeatMapDeformerGraphPath() const
+	{
+		return FString(TEXT("/VertexDeltaModel/Deformers/VertexDeltaHeatMapGraph.VertexDeltaHeatMapGraph"));
+	}
+}	// namespace UE::VertexDeltaModel
+
+#undef LOCTEXT_NAMESPACE
diff --git a/Engine/Plugins/Experimental/MLDeformer/VertexDeltaModel/Source/VertexDeltaModelEditor/Private/VertexDeltaEditorModelActor.cpp b/Engine/Plugins/Experimental/MLDeformer/VertexDeltaModel/Source/VertexDeltaModelEditor/Private/VertexDeltaEditorModelActor.cpp
new file mode 100644
--- /dev/null
+++ b/Engine/Plugins/Experimental/MLDeformer/VertexDeltaModel/Source/VertexDeltaModelEditor/Private/VertexDeltaEditorModelActor.cpp
@@ -0,0 +1,116 @@
+// Copyright Epic Games, Inc. All Rights Reserved.
+
+#include "VertexDeltaEditorModelActor.h"
+#include "GeometryCacheComponent.h"
+#include "Animation/DebugSkelMeshComponent.h"
+
+namespace UE::VertexDeltaModel
+{
+	using namespace UE::MLDeformer;
+
+	FVertexDeltaEditorModelActor::FVertexDeltaEditorModelActor(const FConstructSettings& Settings)
+		: FMLDeformerEditorActor(Settings)
+	{
+	}
+
+	void FVertexDeltaEditorModelActor::SetVisibility(bool bIsVisible)
+	{
+		FMLDeformerEditorActor::SetVisibility(bIsVisible);
+
+		if (GeomCacheComponent && bIsVisible != GeomCacheComponent->IsVisible())
+		{
+			GeomCacheComponent->SetVisibility(bIsVisible, true);
+		}
+	}
+
+	bool FVertexDeltaEditorModelActor::IsVisible() const
+	{
+		if (SkeletalMeshComponent)
+		{
+			return SkeletalMeshComponent->IsVisible();
+		}
+
+		if (GeomCacheComponent)
+		{
+			return GeomCacheComponent->IsVisible();
+		}
+
+		return true;
+	}
+
+	void FVertexDeltaEditorModelActor::SetPlayPosition(float TimeInSeconds, bool bAutoPause)
+	{
+		if (SkeletalMeshComponent)
+		{
+			SkeletalMeshComponent->SetPosition(TimeInSeconds);
+			if (bAutoPause)
+			{
+				SkeletalMeshComponent->bPauseAnims = true;
+			}
+		}
+
+		if (GeomCacheComponent)
+		{
+			GeomCacheComponent->SetManualTick(true);
+			GeomCacheComponent->TickAtThisTime(TimeInSeconds, false, false, false);
+		}
+	}
+
+	float FVertexDeltaEditorModelActor::GetPlayPosition() const
+	{
+		if (SkeletalMeshComponent)
+		{
+			return SkeletalMeshComponent->GetPosition();
+		}
+
+		if (GeomCacheComponent)
+		{
+			return GeomCacheComponent->GetAnimationTime();
+		}
+
+		return 0.0f;
+	}
+
+	void FVertexDeltaEditorModelActor::SetPlaySpeed(float PlaySpeed)
+	{
+		if (SkeletalMeshComponent)
+		{
+			SkeletalMeshComponent->SetPlayRate(PlaySpeed);
+		}
+
+		if (GeomCacheComponent)
+		{
+			GeomCacheComponent->SetPlaybackSpeed(PlaySpeed);
+		}
+	}
+
+	void FVertexDeltaEditorModelActor::Pause(bool bPaused)
+	{
+		if (SkeletalMeshComponent)
+		{
+			SkeletalMeshComponent->bPauseAnims = bPaused;
+		}
+
+		if (GeomCacheComponent)
+		{
+			GeomCacheComponent->SetManualTick(bPaused);
+		}
+	}
+
+	FBox FVertexDeltaEditorModelActor::GetBoundingBox() const
+	{
+		if (SkeletalMeshComponent)
+		{
+			return SkeletalMeshComponent->Bounds.GetBox();
+		}
+
+		if (GeomCacheComponent)
+		{
+			return GeomCacheComponent->Bounds.GetBox();
+		}
+
+		FBox Box;
+		Box.Init();
+		return Box;
+	}
+}	// namespace UE::VertexDeltaModel
diff --git a/Engine/Plugins/Experimental/MLDeformer/VertexDeltaModel/Source/VertexDeltaModelEditor/Private/VertexDeltaModelDetails.cpp b/Engine/Plugins/Experimental/MLDeformer/VertexDeltaModel/Source/VertexDeltaModelEditor/Private/VertexDeltaModelDetails.cpp
new file mode 100644
--- /dev/null
+++ b/Engine/Plugins/Experimental/MLDeformer/VertexDeltaModel/Source/VertexDeltaModelEditor/Private/VertexDeltaModelDetails.cpp
@@ -0,0 +1,186 @@
+// Copyright Epic Games, Inc. All Rights Reserved.
+
+#include "VertexDeltaModelDetails.h"
+#include "VertexDeltaEditorModel.h"
+#include "VertexDeltaModel.h"
+#include "MLDeformerModule.h"
+#include "MLDeformerAsset.h"
+#include "MLDeformerModel.h"
+#include "MLDeformerEditorModel.h"
+#include "MLDeformerEditorModule.h"
+#include "MLDeformerGeomCacheHelpers.h"
+#include "NeuralNetwork.h"
+#include "GeometryCache.h"
+#include "GeometryCacheTrack.h"
+#include "Animation/AnimSequence.h"
+#include "Engine/SkeletalMesh.h"
+#include "PropertyCustomizationHelpers.h"
+#include "AssetRegistry/AssetData.h"
+#include "Modules/ModuleManager.h"
+#include "DetailLayoutBuilder.h"
+#include "DetailCategoryBuilder.h"
+#include "DetailWidgetRow.h"
+#include "IDetailsView.h"
+#include "IDetailGroup.h"
+#include "Widgets/DeclarativeSyntaxSupport.h"
+#include "Widgets/Layout/SBox.h"
+#include "Widgets/Input/SButton.h"
+#include "DetailLayoutBuilder.h"
+#include "SWarningOrErrorBox.h"
+
+#define LOCTEXT_NAMESPACE "VertexDeltaModelDetails"
+
+namespace UE::VertexDeltaModel
+{
+	using namespace UE::MLDeformer;
+
+	TSharedRef<IDetailCustomization> FVertexDeltaModelDetails::MakeInstance()
+	{
+		return MakeShareable(new FVertexDeltaModelDetails());
+	}
+
+	bool FVertexDeltaModelDetails::UpdateMemberPointers(const TArray<TWeakObjectPtr<UObject>>& Objects)
+	{
+		if (!FMLDeformerModelDetails::UpdateMemberPointers(Objects))
+		{
+			return false;
+		}
+
+		VertexModel = Cast<UVertexDeltaModel>(Model);
+		check(VertexModel);
+		VertexEditorModel = static_cast<FVertexDeltaEditorModel*>(EditorModel);
+
+		return (VertexModel != nullptr && VertexEditorModel != nullptr);
+	}
+
+	void FVertexDeltaModelDetails::CustomizeDetails(IDetailLayoutBuilder& DetailBuilder)
+	{
+		// Create all the detail categories and add the properties of the base class.
+		FMLDeformerModelDetails::CustomizeDetails(DetailBuilder);
+
+		// Training settings.
+		SettingsCategoryBuilder->AddProperty(GET_MEMBER_NAME_CHECKED(UVertexDeltaModel, NumHiddenLayers), UMLDeformerModel::StaticClass());
+		SettingsCategoryBuilder->AddProperty(GET_MEMBER_NAME_CHECKED(UVertexDeltaModel, NumNeuronsPerLayer), UMLDeformerModel::StaticClass());
+		SettingsCategoryBuilder->AddProperty(GET_MEMBER_NAME_CHECKED(UVertexDeltaModel, Epochs), UMLDeformerModel::StaticClass());
+
+		// Check whether shrinkage settings should be visible or not.
+		auto IsShrinkageVisible = [this]()
+		{
+			return (VertexModel->GetLossFunction() == EVertexDeltaModelLossFunction::Shrinkage) ? EVisibility::Visible : EVisibility::Collapsed;
+		};
+	 
+		// Advanced settings.
+		SettingsCategoryBuilder->AddProperty(GET_MEMBER_NAME_CHECKED(UVertexDeltaModel, BatchSize), UMLDeformerModel::StaticClass());
+		SettingsCategoryBuilder->AddProperty(GET_MEMBER_NAME_CHECKED(UVertexDeltaModel, LearningRate), UMLDeformerModel::StaticClass());
+		SettingsCategoryBuilder->AddProperty(GET_MEMBER_NAME_CHECKED(UVertexDeltaModel, ActivationFunction), UMLDeformerModel::StaticClass());
+		SettingsCategoryBuilder->AddProperty(GET_MEMBER_NAME_CHECKED(UVertexDeltaModel, LossFunction), UMLDeformerModel::StaticClass());
+		SettingsCategoryBuilder->AddProperty(GET_MEMBER_NAME_CHECKED(UVertexDeltaModel, ShrinkageSpeed), UMLDeformerModel::StaticClass())
+			.Visibility(TAttribute<EVisibility>::CreateLambda(IsShrinkageVisible));
+		SettingsCategoryBuilder->AddProperty(GET_MEMBER_NAME_CHECKED(UVertexDeltaModel, ShrinkageThreshold), UMLDeformerModel::StaticClass())
+			.Visibility(TAttribute<EVisibility>::CreateLambda(IsShrinkageVisible));
+	}
+
+	void FVertexDeltaModelDetails::AddBaseMeshErrors()
+	{
+		// Check if the base mesh matches the target mesh vertex count.
+		FText ErrorText = GetGeomCacheVertexErrorText(Model->SkeletalMesh, VertexModel->GetGeometryCache(), FText::FromString("Base Mesh"), FText::FromString("Target Mesh"));
+		BaseMeshCategoryBuilder->AddCustomRow(FText::FromString("BaseMeshError"))
+			.Visibility(!ErrorText.IsEmpty() ? EVisibility::Visible : EVisibility::Collapsed)
+			.WholeRowContent()
+			[
+				SNew(SBox)
+				.Padding(FMargin(0.0f, 4.0f))
+				[
+					SNew(SWarningOrErrorBox)
+					.MessageStyle(EMessageStyle::Error)
+					.Message(ErrorText)
+				]
+			];
+	}
+	
+	void FVertexDeltaModelDetails::AddAnimSequenceErrors()
+	{
+		const FText WarningText = GetGeomCacheAnimSequenceErrorText(VertexModel->GetGeometryCache(), Model->GetAnimSequence());
+		BaseMeshCategoryBuilder->AddCustomRow(FText::FromString("AnimSeqWarning"))
+			.Visibility(!WarningText.IsEmpty() ? EVisibility::Visible : EVisibility::Collapsed)
+			.WholeRowContent()
+			[
+				SNew(SBox)
+				.Padding(FMargin(0.0f, 4.0f))
+				[
+					SNew(SWarningOrErrorBox)
+					.MessageStyle(EMessageStyle::Warning)
+					.Message(WarningText)
+				]
+			];
+	}
+
+	void FVertexDeltaModelDetails::AddTargetMesh()
+	{
+		TargetMeshCategoryBuilder->AddProperty(GET_MEMBER_NAME_CHECKED(UVertexDeltaModel, GeometryCache), UMLDeformerModel::StaticClass());
+
+		const FText TargetMeshErrorText = GetGeomCacheErrorText(VertexModel->GetSkeletalMesh(), VertexModel->GetGeometryCache());
+		TargetMeshCategoryBuilder->AddCustomRow(FText::FromString("TargetMeshError"))
+			.Visibility(!TargetMeshErrorText.IsEmpty() ? EVisibility::Visible : EVisibility::Collapsed)
+			.WholeRowContent()
+			[
+				SNew(SBox)
+				.Padding(FMargin(0.0f, 4.0f))
+				[
+					SNew(SWarningOrErrorBox)
+					.MessageStyle(EMessageStyle::Error)
+					.Message(TargetMeshErrorText)
+				]
+			];
+
+		const FText ChangedErrorText = EditorModel->GetTargetAssetChangedErrorText();
+		TargetMeshCategoryBuilder->AddCustomRow(FText::FromString("TargetMeshChangedError"))
+			.Visibility(!ChangedErrorText.IsEmpty() ? EVisibility::Visible : EVisibility::Collapsed)
+			.WholeRowContent()
+			[
+				SNew(SBox)
+				.Padding(FMargin(0.0f, 4.0f))
+				[
+					SNew(SWarningOrErrorBox)
+					.MessageStyle(EMessageStyle::Error)
+					.Message(ChangedErrorText)
+				]
+			];
+
+		// Show meshes that have no matching geometry track list.
+		const FText MeshMappingErrorList = GetGeomCacheMeshMappingErrorText(Model->GetSkeletalMesh(), VertexModel->GetGeometryCache());
+		FString GeomTrackNameList;
+		if (!MeshMappingErrorList.IsEmpty())
+		{
+			UGeometryCache* GeomCache = VertexModel->GetGeometryCache();
+			for (int32 Index = 0; Index < GeomCache->Tracks.Num(); ++Index)
+			{
+				GeomTrackNameList += GeomCache->Tracks[Index]->GetName();
+				if (Index < GeomCache->Tracks.Num() - 1)
+				{
+					GeomTrackNameList += TEXT("\n");
+				}
+			}
+		}
+
+		FText MeshMappingErrorFull = FText::Format(
+			LOCTEXT("MeshMappingError", "No matching GeomCache Tracks names found for meshes:\n{0}\n\nGeomCache Track List:\n{1}"), 
+			MeshMappingErrorList,
+			FText::FromString(GeomTrackNameList));
+
+		TargetMeshCategoryBuilder->AddCustomRow(FText::FromString("MeshMappingError"))
+			.Visibility(!MeshMappingErrorList.IsEmpty() ? EVisibility::Visible : EVisibility::Collapsed)
+			.WholeRowContent()
+			[
+				SNew(SBox)
+				.Padding(FMargin(0.0f, 4.0f))
+				[
+					SNew(SWarningOrErrorBox)
+					.MessageStyle(EMessageStyle::Error)
+					.Message(MeshMappingErrorFull)
+				]
+			];
+	}
+}	// namespace UE::VertexDeltaModel
+
+#undef LOCTEXT_NAMESPACE
diff --git a/Engine/Plugins/Experimental/MLDeformer/VertexDeltaModel/Source/VertexDeltaModelEditor/Private/VertexDeltaModelVizSettingsDetails.cpp b/Engine/Plugins/Experimental/MLDeformer/VertexDeltaModel/Source/VertexDeltaModelEditor/Private/VertexDeltaModelVizSettingsDetails.cpp
new file mode 100644
--- /dev/null
+++ b/Engine/Plugins/Experimental/MLDeformer/VertexDeltaModel/Source/VertexDeltaModelEditor/Private/VertexDeltaModelVizSettingsDetails.cpp
@@ -0,0 +1,104 @@
+// Copyright Epic Games, Inc. All Rights Reserved.
+
+#include "VertexDeltaModelVizSettingsDetails.h"
+#include "MLDeformerModule.h"
+#include "MLDeformerModel.h"
+#include "MLDeformerEditorModule.h"
+#include "MLDeformerGeomCacheHelpers.h"
+#include "VertexDeltaModel.h"
+#include "VertexDeltaEditorModel.h"
+#include "VertexDeltaModelVizSettings.h"
+#include "DetailLayoutBuilder.h"
+#include "DetailCategoryBuilder.h"
+#include "DetailWidgetRow.h"
+#include "IDetailsView.h"
+#include "IDetailChildrenBuilder.h"
+#include "IDetailGroup.h"
+#include "Modules/ModuleManager.h"
+#include "Layout/Margin.h"
+#include "Widgets/Layout/SBorder.h"
+#include "Widgets/Text/STextBlock.h"
+#include "Widgets/Layout/SBox.h"
+#include "PropertyCustomizationHelpers.h"
+#include "AssetRegistry/AssetData.h"
+#include "PropertyHandle.h"
+#include "DetailLayoutBuilder.h"
+#include "SWarningOrErrorBox.h"
+#include "GeometryCache.h"
+#include "Animation/AnimSequence.h"
+#include "Animation/MeshDeformer.h"
+#include "Engine/SkeletalMesh.h"
+
+#define LOCTEXT_NAMESPACE "VertexDeltaModelVizSettingsDetails"
+
+namespace UE::VertexDeltaModel
+{
+	using namespace UE::MLDeformer;
+
+	TSharedRef<IDetailCustomization> FVertexDeltaModelVizSettingsDetails::MakeInstance()
+	{
+		return MakeShareable(new FVertexDeltaModelVizSettingsDetails());
+	}
+
+	bool FVertexDeltaModelVizSettingsDetails::UpdateMemberPointers(const TArray<TWeakObjectPtr<UObject>>& Objects)
+	{
+		if (!FMLDeformerVizSettingsDetails::UpdateMemberPointers(Objects))
+		{
+			return false;
+		}
+
+		VertexDeltaModel = Cast<UVertexDeltaModel>(Model);
+		VertexDeltaVizSettings = Cast<UVertexDeltaModelVizSettings>(VizSettings);
+		return (VertexDeltaModel != nullptr && VertexDeltaVizSettings != nullptr);
+	}
+
+	void FVertexDeltaModelVizSettingsDetails::AddGroundTruth()
+	{
+		TestAssetsCategory->AddProperty(GET_MEMBER_NAME_CHECKED(UVertexDeltaModelVizSettings, GroundTruth));
+
+		// Show an error when the test anim sequence duration doesn't match the one of the ground truth.
+		const FText AnimErrorText = GetGeomCacheAnimSequenceErrorText(VertexDeltaVizSettings->GetTestGroundTruth(), VizSettings->GetTestAnimSequence());
+		FDetailWidgetRow& GroundTruthAnimErrorRow = TestAssetsCategory->AddCustomRow(FText::FromString("GroundTruthAnimMismatchError"))
+			.Visibility(!AnimErrorText.IsEmpty() ? EVisibility::Visible : EVisibility::Collapsed)
+			.WholeRowContent()
+			[
+				SNew(SBox)
+				.Padding(FMargin(0.0f, 4.0f))
+				[
+					SNew(SWarningOrErrorBox)
+					.MessageStyle(EMessageStyle::Warning)
+					.Message(AnimErrorText)
+				]
+			];
+
+		const FText GeomErrorText = GetGeomCacheErrorText(Model->GetSkeletalMesh(), VertexDeltaVizSettings->GetTestGroundTruth());
+		FDetailWidgetRow& GroundTruthGeomErrorRow = TestAssetsCategory->AddCustomRow(FText::FromString("GroundTruthGeomMismatchError"))
+			.Visibility(!GeomErrorText.IsEmpty() ? EVisibility::Visible : EVisibility::Collapsed)
+			.WholeRowContent()
+			[
+				SNew(SBox)
+				.Padding(FMargin(0.0f, 4.0f))
+				[
+					SNew(SWarningOrErrorBox)
+					.MessageStyle(EMessageStyle::Warning)
+					.Message(GeomErrorText)
+				]
+			];
+
+		const FText VertexErrorText = GetGeomCacheVertexErrorText(Model->GetSkeletalMesh(), VertexDeltaVizSettings->GetTestGroundTruth(), FText::FromString("Base Mesh"), FText::FromString("Ground Truth Mesh"));
+		FDetailWidgetRow& GroundTruthVertexErrorRow = TestAssetsCategory->AddCustomRow(FText::FromString("GroundTruthVertexMismatchError"))
+			.Visibility(!VertexErrorText.IsEmpty() ? EVisibility::Visible : EVisibility::Collapsed)
+			.WholeRowContent()
+			[
+				SNew(SBox)
+				.Padding(FMargin(0.0f, 4.0f))
+				[
+					SNew(SWarningOrErrorBox)
+					.MessageStyle(EMessageStyle::Warning)
+					.Message(VertexErrorText)
+				]
+			];
+	}
+}	//namespace UE::VertexDeltaModel
+
+#undef LOCTEXT_NAMESPACE
diff --git a/Engine/Plugins/Experimental/MLDeformer/VertexDeltaModel/Source/VertexDeltaModelEditor/Private/VertexDeltaTrainingModel.cpp b/Engine/Plugins/Experimental/MLDeformer/VertexDeltaModel/Source/VertexDeltaModelEditor/Private/VertexDeltaTrainingModel.cpp
new file mode 100644
--- /dev/null
+++ b/Engine/Plugins/Experimental/MLDeformer/VertexDeltaModel/Source/VertexDeltaModelEditor/Private/VertexDeltaTrainingModel.cpp
@@ -0,0 +1,117 @@
+// Copyright Epic Games, Inc. All Rights Reserved.
+
+#include "VertexDeltaTrainingModel.h"
+#include "VertexDeltaModel.h"
+#include "VertexDeltaEditorModel.h"
+#include "MLDeformerEditorModel.h"
+#include "Misc/ScopedSlowTask.h"
+#include "Math/NumericLimits.h"
+
+#define LOCTEXT_NAMESPACE "VertexDeltaTrainingModel"
+
+UVertexDeltaModel* UVertexDeltaTrainingModel::GetVertexDeltaModel() const
+{
+	return Cast<UVertexDeltaModel>(GetModel());
+}
+
+UE::VertexDeltaModel::FVertexDeltaEditorModel* UVertexDeltaTrainingModel::GetVertexDeltaEditorModel() const
+{
+	return static_cast<UE::VertexDeltaModel::FVertexDeltaEditorModel*>(EditorModel);
+}
+
+void UVertexDeltaTrainingModel::UpdateVertexDeltaMeanAndScale(const TArray<float>& VertexDeltas, FVector& InOutMeanVertexDelta, FVector& InOutVertexDeltaScale, float& InOutCount)
+{
+	// Calculate this frame's mean.
+	FVector MeanDelta = FVector::ZeroVector;
+	FVector MinDelta(TNumericLimits<float>::Max());
+	FVector MaxDelta(-TNumericLimits<float>::Max());
+	const int32 NumVertices = VertexDeltas.Num() / 3;
+	for (int32 VertexIndex = 0; VertexIndex < NumVertices; ++VertexIndex)
+	{
+		const int Offset = VertexIndex * 3;
+		const FVector Delta(VertexDeltas[Offset], VertexDeltas[Offset+1], VertexDeltas[Offset+2]);
+		MeanDelta += Delta;
+		MinDelta = MinDelta.ComponentMin(Delta);
+		MaxDelta = MaxDelta.ComponentMax(Delta);
+	}
+	const float ValidVertexCount = (float)NumVertices;
+	if (NumVertices > 0)
+	{
+		MeanDelta /= (float)NumVertices;
+	}
+
+	// Update global mean.
+	InOutCount += 1.0f;
+	FVector TempDiff = MeanDelta - (FVector)InOutMeanVertexDelta;
+	InOutMeanVertexDelta = InOutMeanVertexDelta + FVector(TempDiff / InOutCount);
+
+	// Update global scale.
+	TempDiff = MaxDelta - MinDelta;
+	InOutVertexDeltaScale = InOutVertexDeltaScale.ComponentMax(TempDiff.GetAbs());
+}
+
+bool UVertexDeltaTrainingModel::ComputeVertexDeltaStatistics(uint32 LODIndex)
+{
+	using namespace UE::VertexDeltaModel;
+
+	FVertexDeltaEditorModel* VertexDeltaEditorModel = GetVertexDeltaEditorModel();
+	UVertexDeltaModel* VertexDeltaModel = VertexDeltaEditorModel->GetVertexDeltaModel();
+	if (VertexDeltaEditorModel->IsDataNormalized())
+	{
+		return true;
+	}
+
+	FMLDeformerSampler* Sampler = VertexDeltaEditorModel->GetSampler();
+	Sampler->SetVertexDeltaSpace(EVertexDeltaSpace::PreSkinning);
+
+	// Show some dialog with progress.
+	const FText Title(LOCTEXT("PreprocessTrainingDataMessage", "Calculating data statistics"));
+	const uint32 AnimNumFrames = EditorModel->GetNumFramesForTraining();
+	FScopedSlowTask Task((float)AnimNumFrames, Title);
+	Task.MakeDialog(true);
+
+	// Get the base actor data.
+	FVector Mean = FVector::ZeroVector;
+	FVector Scale = FVector::OneVector;
+	float Count = 0.0f;
+	const int32 FrameCount = VertexDeltaEditorModel->GetNumFramesForTraining();
+	for (int32 FrameIndex = 0; FrameIndex < FrameCount; FrameIndex++)
+	{
+		// Calculate the deltas and update the mean and scale.
+		Sampler->Sample(FrameIndex);
+		UpdateVertexDeltaMeanAndScale(Sampler->GetVertexDeltas(), Mean, Scale, Count);
+
+		// Forward the progress bar and check if we want to cancel.
+		Task.EnterProgressFrame();
+		if (Task.ShouldCancel())
+		{
+			return false;
+		}
+	}
+
+	// Update the asset with calculated statistics.
+	VertexDeltaModel->VertexDeltaMean = Mean;
+	if (Count > 0.0f)
+	{
+		VertexDeltaModel->VertexDeltaScale = FVector::OneVector * Scale.GetMax();
+		VertexDeltaEditorModel->SetDataNormalized(true);
+	}
+
+	return true;
+}
+
+bool UVertexDeltaTrainingModel::ComputeDeltasStatistics()
+{
+	const UVertexDeltaModel* VertexDeltaModel = GetVertexDeltaModel();
+	const float DeltaCutoffLength = VertexDeltaModel->GetDeltaCutoffLength();
+	if (ComputeVertexDeltaStatistics(0))
+	{
+		VertexDeltaMean = VertexDeltaModel->VertexDeltaMean;
+		VertexDeltaScale = VertexDeltaModel->VertexDeltaScale;
+		return true;
+	}
+
+	return false;
+}
+
+#undef LOCTEXT_NAMESPACE
diff --git a/Engine/Plugins/Experimental/MLDeformer/VertexDeltaModel/Source/VertexDeltaModelEditor/Public/VertexDeltaEditorModel.h b/Engine/Plugins/Experimental/MLDeformer/VertexDeltaModel/Source/VertexDeltaModelEditor/Public/VertexDeltaEditorModel.h
new file mode 100644
--- /dev/null
+++ b/Engine/Plugins/Experimental/MLDeformer/VertexDeltaModel/Source/VertexDeltaModelEditor/Public/VertexDeltaEditorModel.h
@@ -0,0 +1,73 @@
+// Copyright Epic Games, Inc. All Rights Reserved.
+
+#pragma once
+
+#include "CoreMinimal.h"
+#include "MLDeformerEditorModel.h"
+#include "MLDeformerGeomCacheSampler.h"
+#include "VertexDeltaModel.h"
+#include "MLDeformerEditorActor.h"
+#include "UObject/GCObject.h"
+
+class UMLDeformerAsset;
+class USkeletalMesh;
+class UGeometryCache;
+class UVertexDeltaModel;
+class UVertexDeltaModelVizSettings;
+
+namespace UE::VertexDeltaModel
+{
+	using namespace UE::MLDeformer;
+
+	class FVertexDeltaEditorModelActor;
+	class FVertexDeltaModelSampler;
+
+	class VERTEXDELTAMODELEDITOR_API FVertexDeltaEditorModel 
+		: public UE::MLDeformer::FMLDeformerEditorModel
+	{
+	public:
+		virtual ~FVertexDeltaEditorModel() override;
+
+		// We need to implement this static MakeInstance method.
+		static FMLDeformerEditorModel* MakeInstance();
+
+		// FGCObject overrides.
+		virtual void AddReferencedObjects(FReferenceCollector& Collector) override;
+		virtual FString GetReferencerName() const override { return TEXT("FVertexDeltaEditorModel"); }
+		// ~END FGCObject overrides.
+	
+		// FMLDeformerEditorModel overrides.
+		virtual void Init(const InitSettings& Settings) override;
+		virtual FMLDeformerEditorActor* CreateEditorActor(const FMLDeformerEditorActor::FConstructSettings& Settings) const override;
+		virtual FMLDeformerSampler* CreateSampler() const;
+		virtual void CreateTrainingGroundTruthActor(UWorld* World) override;
+		virtual void CreateTestGroundTruthActor(UWorld* World) override;
+		virtual void OnPreTraining() override;
+		virtual void OnTrainingAborted() override;
+		virtual void OnInputAssetsChanged() override;
+		virtual int32 GetNumFrames() const override;
+		virtual double GetTimeAtFrame(int32 FrameNumber) const override;
+		virtual int32 GetFrameAtTime(double TimeInSeconds) const override;
+		virtual void UpdateIsReadyForTrainingState() override;
+		virtual ETrainingResult Train() override;
+		virtual FString GetDefaultDeformerGraphAssetPath() const override;
+		virtual FString GetHeatMapDeformerGraphPath() const override;
+		FString GetTrainedNetworkOnnxFile() const override;
+		// ~END FMLDeformerEditorModel overrides.
+
+		// Some helpers that cast to this model's variants of some classes.
+		UVertexDeltaModel* GetVertexDeltaModel() const { return Cast<UVertexDeltaModel>(Model); }
+		UVertexDeltaModelVizSettings* GetVertexDeltaModelVizSettings() const;
+		FVertexDeltaEditorModelActor* FindVertexDeltaModelEditorActor(int32 TypeID) const;
+		FMLDeformerGeomCacheSampler* GetGeomCacheSampler() const { return static_cast<FMLDeformerGeomCacheSampler*>(Sampler); }
+	
+		void OnPostEditChangeProperty(FPropertyChangedEvent& PropertyChangedEvent);
+
+	protected:
+		void CreateGeomCacheActor(UWorld* World, int32 ActorID, const FName& Name, UGeometryCache* GeomCache, FLinearColor LabelColor, FLinearColor WireframeColor, const FText& LabelText, bool bIsTrainingActor);
+
+	protected:
+		FVector VertexDeltaMeanBackup = FVector::ZeroVector;
+		FVector VertexDeltaScaleBackup = FVector::OneVector;
+	};
+}	// namespace UE::VertexDeltaModel
diff --git a/Engine/Plugins/Experimental/MLDeformer/VertexDeltaModel/Source/VertexDeltaModelEditor/Public/VertexDeltaEditorModelActor.h b/Engine/Plugins/Experimental/MLDeformer/VertexDeltaModel/Source/VertexDeltaModelEditor/Public/VertexDeltaEditorModelActor.h
new file mode 100644
--- /dev/null
+++ b/Engine/Plugins/Experimental/MLDeformer/VertexDeltaModel/Source/VertexDeltaModelEditor/Public/VertexDeltaEditorModelActor.h
@@ -0,0 +1,37 @@
+// Copyright Epic Games, Inc. All Rights Reserved.
+
+#pragma once
+
+#include "MLDeformerEditorActor.h"
+
+class UMLDeformerComponent;
+class UGeometryCacheComponent;
+
+namespace UE::VertexDeltaModel
+{
+	using namespace UE::MLDeformer;
+
+	class VERTEXDELTAMODELEDITOR_API FVertexDeltaEditorModelActor
+		: public FMLDeformerEditorActor
+	{
+	public:
+		FVertexDeltaEditorModelActor(const FConstructSettings& Settings);
+
+		void SetGeometryCacheComponent(UGeometryCacheComponent* Component) { GeomCacheComponent = Component; }
+		UGeometryCacheComponent* GetGeometryCacheComponent() const { return GeomCacheComponent; }
+
+		// FMLDeformerEditorActor overrides.
+		virtual void SetVisibility(bool bIsVisible) override;
+		virtual bool IsVisible() const override;
+		virtual void SetPlayPosition(float TimeInSeconds, bool bAutoPause=true) override;
+		virtual float GetPlayPosition() const override;
+		virtual void SetPlaySpeed(float PlaySpeed) override;
+		virtual void Pause(bool bPaused) override;
+		virtual FBox GetBoundingBox() const override;
+		// ~END FMLDeformerEditorActor overrides.
+
+	protected:
+		/** The geometry cache component (can be nullptr). */
+		TObjectPtr<UGeometryCacheComponent> GeomCacheComponent = nullptr;
+	};
+}	// namespace UE::VertexDeltaModel
diff --git a/Engine/Plugins/Experimental/MLDeformer/VertexDeltaModel/Source/VertexDeltaModelEditor/Public/VertexDeltaModelDetails.h b/Engine/Plugins/Experimental/MLDeformer/VertexDeltaModel/Source/VertexDeltaModelEditor/Public/VertexDeltaModelDetails.h
new file mode 100644
--- /dev/null
+++ b/Engine/Plugins/Experimental/MLDeformer/VertexDeltaModel/Source/VertexDeltaModelEditor/Public/VertexDeltaModelDetails.h
@@ -0,0 +1,45 @@
+// Copyright Epic Games, Inc. All Rights Reserved.
+
+#pragma once
+
+#include "CoreMinimal.h"
+#include "MLDeformerModelDetails.h"
+#include "IDetailCustomization.h"
+#include "AssetRegistry/AssetData.h"
+#include "Input/Reply.h"
+
+class IDetailLayoutBuilder;
+class UVertexDeltaModel;
+
+namespace UE::MLDeformer
+{
+	class FMLDeformerEditorModel;
+}
+
+namespace UE::VertexDeltaModel
+{
+	class FVertexDeltaEditorModel;
+
+	class VERTEXDELTAMODELEDITOR_API FVertexDeltaModelDetails final
+		: public UE::MLDeformer::FMLDeformerModelDetails
+	{
+	public:
+		/** Makes a new instance of this detail layout class for a specific detail view requesting it. */
+		static TSharedRef<IDetailCustomization> MakeInstance();
+
+		// ILayoutDetails overrides.
+		void CustomizeDetails(IDetailLayoutBuilder& DetailBuilder) override;
+		// ~END ILayoutDetails overrides.
+
+		// FMLDeformerModelDetails overrides.
+		bool UpdateMemberPointers(const TArray<TWeakObjectPtr<UObject>>& Objects) override;
+		void AddTargetMesh() override;
+		void AddBaseMeshErrors() override;
+		void AddAnimSequenceErrors() override;
+		// ~END FMLDeformerModelDetails overrides.
+
+	protected:
+		TObjectPtr<UVertexDeltaModel> VertexModel = nullptr;
+		FVertexDeltaEditorModel* VertexEditorModel = nullptr;
+	};
+}	// namespace UE::VertexDeltaModel
diff --git a/Engine/Plugins/Experimental/MLDeformer/VertexDeltaModel/Source/VertexDeltaModelEditor/Public/VertexDeltaModelEditorModule.h b/Engine/Plugins/Experimental/MLDeformer/VertexDeltaModel/Source/VertexDeltaModelEditor/Public/VertexDeltaModelEditorModule.h
new file mode 100644
--- /dev/null
+++ b/Engine/Plugins/Experimental/MLDeformer/VertexDeltaModel/Source/VertexDeltaModelEditor/Public/VertexDeltaModelEditorModule.h
@@ -0,0 +1,19 @@
+// Copyright Epic Games, Inc. All Rights Reserved.
+
+#pragma once
+
+#include "CoreMinimal.h"
+#include "Modules/ModuleInterface.h"
+
+namespace UE::VertexDeltaModel
+{
+	class VERTEXDELTAMODELEDITOR_API FVertexDeltaModelEditorModule
+		: public IModuleInterface
+	{
+	public:
+		// IModuleInterface overrides.
+		void StartupModule() override;
+		void ShutdownModule() override;
+		// ~END IModuleInterface overrides.
+	};
+}	// namespace UE::VertexDeltaModel
\ No newline at end of file
diff --git a/Engine/Plugins/Experimental/MLDeformer/VertexDeltaModel/Source/VertexDeltaModelEditor/Public/VertexDeltaModelVizSettingsDetails.h b/Engine/Plugins/Experimental/MLDeformer/VertexDeltaModel/Source/VertexDeltaModelEditor/Public/VertexDeltaModelVizSettingsDetails.h
new file mode 100644
--- /dev/null
+++ b/Engine/Plugins/Experimental/MLDeformer/VertexDeltaModel/Source/VertexDeltaModelEditor/Public/VertexDeltaModelVizSettingsDetails.h
@@ -0,0 +1,34 @@
+// Copyright Epic Games, Inc. All Rights Reserved.
+
+#pragma once
+
+#include "CoreMinimal.h"
+#include "MLDeformerVizSettingsDetails.h"
+#include "IDetailCustomization.h"
+#include "AssetRegistry/AssetData.h"
+#include "PropertyHandle.h"
+
+class IDetailLayoutBuilder;
+class USkeleton;
+class UVertexDeltaModel;
+class UVertexDeltaModelVizSettings;
+
+namespace UE::VertexDeltaModel
+{
+	class VERTEXDELTAMODELEDITOR_API FVertexDeltaModelVizSettingsDetails final
+		: public UE::MLDeformer::FMLDeformerVizSettingsDetails
+	{
+	public:
+		/** Makes a new instance of this detail layout class for a specific detail view requesting it. */
+		static TSharedRef<IDetailCustomization> MakeInstance();
+
+		// FMLDeformerVizSettingsDetails overrides.
+		bool UpdateMemberPointers(const TArray<TWeakObjectPtr<UObject>>& Objects) override;
+		void AddGroundTruth() override;
+		// ~END FMLDeformerVizSettingsDetails overrides.
+
+	protected:
+		TObjectPtr<UVertexDeltaModel> VertexDeltaModel = nullptr;
+		TObjectPtr<UVertexDeltaModelVizSettings> VertexDeltaVizSettings = nullptr;
+	};
+}	// namespace UE::VertexDeltaModel
diff --git a/Engine/Plugins/Experimental/MLDeformer/VertexDeltaModel/Source/VertexDeltaModelEditor/Public/VertexDeltaTrainingModel.h b/Engine/Plugins/Experimental/MLDeformer/VertexDeltaModel/Source/VertexDeltaModelEditor/Public/VertexDeltaTrainingModel.h
new file mode 100644
--- /dev/null
+++ b/Engine/Plugins/Experimental/MLDeformer/VertexDeltaModel/Source/VertexDeltaModelEditor/Public/VertexDeltaTrainingModel.h
@@ -0,0 +1,46 @@
+// Copyright Epic Games, Inc. All Rights Reserved.
+
+#pragma once
+
+#include "CoreMinimal.h"
+#include "MLDeformerTrainingModel.h"
+#include "VertexDeltaTrainingModel.generated.h"
+
+namespace UE::VertexDeltaModel
+{
+	class FVertexDeltaEditorModel;
+}
+
+class UVertexDeltaModel;
+
+UCLASS(Blueprintable)
+class VERTEXDELTAMODELEDITOR_API UVertexDeltaTrainingModel
+	: public UMLDeformerTrainingModel
+{
+	GENERATED_BODY()
+
+public:
+	/** Main training function, with implementation in python. */
+	UFUNCTION(BlueprintImplementableEvent, Category = "Python")
+	int32 Train() const;
+
+	/** Compute delta statistics for the whole dataset. */
+	UFUNCTION(BlueprintCallable, Category = "Training Data")
+	bool ComputeDeltasStatistics();
+
+	UVertexDeltaModel* GetVertexDeltaModel() const;
+	UE::VertexDeltaModel::FVertexDeltaEditorModel* GetVertexDeltaEditorModel() const;
+
+protected:
+	bool ComputeVertexDeltaStatistics(uint32 LODIndex);
+	void UpdateVertexDeltaMeanAndScale(const TArray<float>& VertexDeltas, FVector& InOutMeanVertexDelta, FVector& InOutVertexDeltaScale, float& InOutCount);
+
+public:
+	// Mean delta computed over the entire dataset.
+	UPROPERTY(EditAnywhere, BlueprintReadWrite, Category = "Training Data")
+	FVector VertexDeltaMean = FVector::ZeroVector;
+
+	// Vertex delta scale computed over the entire dataset.
+	UPROPERTY(EditAnywhere, BlueprintReadWrite, Category = "Training Data")
+	FVector VertexDeltaScale = FVector::OneVector;
+};
diff --git a/Engine/Plugins/Experimental/MLDeformer/VertexDeltaModel/VertexDeltaModel.uplugin b/Engine/Plugins/Experimental/MLDeformer/VertexDeltaModel/VertexDeltaModel.uplugin
new file mode 100644
--- /dev/null
+++ b/Engine/Plugins/Experimental/MLDeformer/VertexDeltaModel/VertexDeltaModel.uplugin
@@ -0,0 +1,51 @@
+{
+	"FileVersion": 3,
+	"Version": 1,
+	"VersionName": "0.2",
+	"FriendlyName": "VertexDeltaModel",
+	"Description": "Vertex Delta Model for the ML Deformer Framework",
+	"Category": "Animation",
+	"CreatedBy": "Epic Games, Inc.",
+	"CreatedByURL": "http://epicgames.com",
+	"DocsURL": "",
+	"MarketplaceURL": "",
+	"SupportURL": "",
+	"CanContainContent": true,
+	"IsBetaVersion": false,
+	"IsExperimentalVersion": true,
+	"Installed": false,
+	"Modules": [
+		{
+			"Name": "VertexDeltaModel",
+			"Type": "Runtime",
+			"LoadingPhase": "Default"
+		},
+		{
+			"Name": "VertexDeltaModelEditor",
+			"Type": "Editor",
+			"LoadingPhase": "Default"
+		}
+	],
+	"Plugins": [
+		{
+			"Name": "GeometryCache",
+			"Enabled": true
+		},
+		{
+			"Name": "NeuralNetworkInference",
+			"Enabled": true
+		},
+		{
+			"Name": "Optimus",
+			"Enabled": true
+		},
+		{
+			"Name": "PythonFoundationPackages",
+			"Enabled": true
+		},
+		{
+			"Name": "MLDeformerFramework",
+			"Enabled": true
+		}
+	]
+}
\ No newline at end of file
