[MLDeformer]
- Fixed a crash after training and closing the editor.
- Changed the output onnx name to VertexDeltaModel.onnx
- Added estimated time left and passed time to the training and sampling.
- Some minor python code tweaks.
- Aborting the frame sampling now doesn't offer the user to use the partially trained network, as nothing was trained yet.

diff --git a/Engine/Plugins/Experimental/MLDeformer/MLDeformerFramework/Source/MLDeformerFramework/Private/MLDeformerComponent.cpp b/Engine/Plugins/Experimental/MLDeformer/MLDeformerFramework/Source/MLDeformerFramework/Private/MLDeformerComponent.cpp
--- a/Engine/Plugins/Experimental/MLDeformer/MLDeformerFramework/Source/MLDeformerFramework/Private/MLDeformerComponent.cpp
+++ b/Engine/Plugins/Experimental/MLDeformer/MLDeformerFramework/Source/MLDeformerFramework/Private/MLDeformerComponent.cpp
@@ -28,6 +28,10 @@ void UMLDeformerComponent::Init()
 	UMLDeformerModel* Model = DeformerAsset->GetModel();
 	if (Model)
 	{
+		if (ModelInstance)
+		{
+			ModelInstance->Release();
+		}
 		ModelInstance = Model->CreateModelInstance(this);
 		ModelInstance->SetModel(Model);
 		ModelInstance->Init(SkelMeshComponent);
diff --git a/Engine/Plugins/Experimental/MLDeformer/VertexDeltaModel/Source/VertexDeltaModelEditor/Private/VertexDeltaEditorModel.cpp b/Engine/Plugins/Experimental/MLDeformer/VertexDeltaModel/Source/VertexDeltaModelEditor/Private/VertexDeltaEditorModel.cpp
--- a/Engine/Plugins/Experimental/MLDeformer/VertexDeltaModel/Source/VertexDeltaModelEditor/Private/VertexDeltaEditorModel.cpp
+++ b/Engine/Plugins/Experimental/MLDeformer/VertexDeltaModel/Source/VertexDeltaModelEditor/Private/VertexDeltaEditorModel.cpp
@@ -271,7 +271,7 @@ namespace UE::VertexDeltaModel
 
 	FString FVertexDeltaEditorModel::GetTrainedNetworkOnnxFile() const
 	{
-		return FString(FPaths::ProjectIntermediateDir() + TEXT("VertexDeltaModel/latest_net_G.onnx"));
+		return FString(FPaths::ProjectIntermediateDir() + TEXT("VertexDeltaModel/VertexDeltaModel.onnx"));
 	}
 }	// namespace UE::VertexDeltaModel
 
