[MLDeformer]
- Move NNI out of the base framework plugin.
- Move NNI out of neural morph model plugin.
- Move NNI out of the JiggleSpring model plugin.
- Move NNI into the Vertex Delta Model plugin.
- Move NNI into the Nearest Neighbor Model plugin.
- Vertex Delta Model now contains its own deformer graphs and shader.

This essentially removes the dependency of NNI inside the ML Deformer framework and neural morph model, and moves it to the model plugins that require it.
Please note that all models will have to be retrained because of this change. If  you use a vertex delta model asset, you need to point it to use the new DG_VertexDeltaModel deformer graph.

diff --git a/Engine/Plugins/Animation/MLDeformer/MLDeformerFramework/MLDeformerFramework.uplugin b/Engine/Plugins/Animation/MLDeformer/MLDeformerFramework/MLDeformerFramework.uplugin
--- a/Engine/Plugins/Animation/MLDeformer/MLDeformerFramework/MLDeformerFramework.uplugin
+++ b/Engine/Plugins/Animation/MLDeformer/MLDeformerFramework/MLDeformerFramework.uplugin
@@ -1,7 +1,7 @@
 {
 	"FileVersion": 3,
 	"Version": 1,
-	"VersionName": "0.1",
+	"VersionName": "0.2",
 	"FriendlyName": "ML Deformer Framework",
 	"Description": "Machine Learning Mesh Deformer Framework",
 	"Category": "Animation",
@@ -35,10 +35,6 @@
 			"Name": "PythonFoundationPackages",
 			"Enabled": true
 		},
-		{
-			"Name": "NeuralNetworkInference",
-			"Enabled": true
-		},
 		{
 			"Name": "GeometryCache",
 			"Enabled": true
diff --git a/Engine/Plugins/Animation/MLDeformer/MLDeformerFramework/Shaders/Private/MLDeformerGraphDataInterface.ush b/Engine/Plugins/Animation/MLDeformer/MLDeformerFramework/Shaders/Private/MLDeformerGraphDataInterface.ush
deleted file mode 100644
diff --git a/Engine/Plugins/Animation/MLDeformer/MLDeformerFramework/Shaders/Private/MLDeformerGraphHeatMapDataInterface.ush b/Engine/Plugins/Animation/MLDeformer/MLDeformerFramework/Shaders/Private/MLDeformerGraphHeatMapDataInterface.ush
deleted file mode 100644
diff --git a/Engine/Plugins/Animation/MLDeformer/MLDeformerFramework/Source/MLDeformerFramework/MLDeformerFramework.Build.cs b/Engine/Plugins/Animation/MLDeformer/MLDeformerFramework/Source/MLDeformerFramework/MLDeformerFramework.Build.cs
--- a/Engine/Plugins/Animation/MLDeformer/MLDeformerFramework/Source/MLDeformerFramework/MLDeformerFramework.Build.cs
+++ b/Engine/Plugins/Animation/MLDeformer/MLDeformerFramework/Source/MLDeformerFramework/MLDeformerFramework.Build.cs
@@ -39,7 +39,6 @@ public class MLDeformerFramework : ModuleRules
 				"RenderCore",
 				"RHI",
 				"GeometryCache",
-				"NeuralNetworkInference",
 				"OptimusCore",
 				"ComputeFramework",
 				"Projects"
diff --git a/Engine/Plugins/Animation/MLDeformer/MLDeformerFramework/Source/MLDeformerFramework/Private/MLDeformerGraphDataInterface.cpp b/Engine/Plugins/Animation/MLDeformer/MLDeformerFramework/Source/MLDeformerFramework/Private/MLDeformerGraphDataInterface.cpp
deleted file mode 100644
diff --git a/Engine/Plugins/Animation/MLDeformer/MLDeformerFramework/Source/MLDeformerFramework/Private/MLDeformerGraphDebugDataInterface.cpp b/Engine/Plugins/Animation/MLDeformer/MLDeformerFramework/Source/MLDeformerFramework/Private/MLDeformerGraphDebugDataInterface.cpp
deleted file mode 100644
diff --git a/Engine/Plugins/Animation/MLDeformer/MLDeformerFramework/Source/MLDeformerFramework/Private/MLDeformerModel.cpp b/Engine/Plugins/Animation/MLDeformer/MLDeformerFramework/Source/MLDeformerFramework/Private/MLDeformerModel.cpp
--- a/Engine/Plugins/Animation/MLDeformer/MLDeformerFramework/Source/MLDeformerFramework/Private/MLDeformerModel.cpp
+++ b/Engine/Plugins/Animation/MLDeformer/MLDeformerFramework/Source/MLDeformerFramework/Private/MLDeformerModel.cpp
@@ -10,7 +10,6 @@
 #include "Engine/SkeletalMesh.h"
 #include "Rendering/SkeletalMeshModel.h"
 #include "Rendering/SkeletalMeshLODModel.h"
-#include "NeuralNetwork.h"
 #include "Animation/AnimSequence.h"
 #include "UObject/UObjectGlobals.h"
 
@@ -124,35 +123,11 @@ void UMLDeformerModel::PostLoad()
 		InputInfo->OnPostLoad();
 	}
 
-	UNeuralNetwork* Network = GetNeuralNetwork();
-	if (Network)
-	{
-		// If we run the neural network on the GPU.
-		if (IsNeuralNetworkOnGPU())
-		{
-			Network->SetDeviceType(ENeuralDeviceType::GPU, ENeuralDeviceType::CPU, ENeuralDeviceType::GPU);
-			if (Network->GetDeviceType() != ENeuralDeviceType::GPU || Network->GetOutputDeviceType() != ENeuralDeviceType::GPU || Network->GetInputDeviceType() != ENeuralDeviceType::CPU)
-			{
-				UE_LOG(LogMLDeformer, Error, TEXT("Neural net in ML Deformer '%s' cannot run on the GPU, it will not be active."), *GetDeformerAsset()->GetName());
-			}
-		}
-		else // We run our neural network on the CPU.
-		{
-			Network->SetDeviceType(ENeuralDeviceType::CPU, ENeuralDeviceType::CPU, ENeuralDeviceType::CPU);
-		}
-	}
-
 	#if WITH_EDITOR
 		UpdateMemoryUsage();
 	#endif
 }
 
-void UMLDeformerModel::SetNeuralNetwork(UNeuralNetwork* InNeuralNetwork)
-{
-	NeuralNetworkModifyDelegate.Broadcast();
-	NeuralNetwork = InNeuralNetwork;
-}
-
 // Used for the FBoenReference, so it knows what skeleton to pick bones from.
 USkeleton* UMLDeformerModel::GetSkeleton(bool& bInvalidSkeletonIsError, const IPropertyHandle* PropertyHandle)
 {
@@ -270,19 +245,6 @@ void UMLDeformerModel::FloatArrayToVector3Array(const TArray<float>& FloatArray,
 		{
 			GPUMemUsageInBytes += VertexMapBuffer.VertexBufferRHI->GetSize();
 		}
-
-		// Check if the neural network is on the GPU, if so, count the memory to the GPU and remove it from Main memory, as we added it to 
-		// that already when we did the Model->GetResourceSizeBytes.
-		if (NeuralNetwork)
-		{
-			if (NeuralNetwork->GetDeviceType() == ENeuralDeviceType::GPU)
-			{
-				const uint64 NeuralNetSize = static_cast<uint64>(NeuralNetwork->GetResourceSizeBytes(EResourceSizeMode::Type::EstimatedTotal));
-				GPUMemUsageInBytes += NeuralNetSize;
-				MemUsageInBytes -= NeuralNetSize;
-				CookedMemUsageInBytes -= NeuralNetSize;
-			}
-		}
 	}
 #endif	// #if WITH_EDITOR
 
diff --git a/Engine/Plugins/Animation/MLDeformer/MLDeformerFramework/Source/MLDeformerFramework/Private/MLDeformerModelInstance.cpp b/Engine/Plugins/Animation/MLDeformer/MLDeformerFramework/Source/MLDeformerFramework/Private/MLDeformerModelInstance.cpp
--- a/Engine/Plugins/Animation/MLDeformer/MLDeformerFramework/Source/MLDeformerFramework/Private/MLDeformerModelInstance.cpp
+++ b/Engine/Plugins/Animation/MLDeformer/MLDeformerFramework/Source/MLDeformerFramework/Private/MLDeformerModelInstance.cpp
@@ -9,7 +9,6 @@
 #include "Components/SkeletalMeshComponent.h"
 #include "Engine/SkeletalMesh.h"
 #include "Animation/AnimInstance.h"
-#include "NeuralNetwork.h"
 
 #include UE_INLINE_GENERATED_CPP_BY_NAME(MLDeformerModelInstance)
 
@@ -25,14 +24,6 @@ void UMLDeformerModelInstance::Release()
 	// that we are about to delete.
 	RenderCommandFence.BeginFence();
 	RenderCommandFence.Wait();
-
-	// Destroy the neural network instance.
-	UNeuralNetwork* NeuralNetwork = Model.Get() ? Model->GetNeuralNetwork() : nullptr;
-	if (NeuralNetwork && NeuralNetworkInferenceHandle != -1)
-	{
-		NeuralNetwork->DestroyInferenceContext(NeuralNetworkInferenceHandle);
-		NeuralNetworkInferenceHandle = -1;
-	}
 }
 
 USkeletalMeshComponent* UMLDeformerModelInstance::GetSkeletalMeshComponent() const
@@ -126,7 +117,7 @@ void UMLDeformerModelInstance::UpdateCompatibilityStatus()
 	bIsCompatible = SkeletalMeshComponent->GetSkeletalMeshAsset() && CheckCompatibility(SkeletalMeshComponent, true).IsEmpty();
 }
 
-FString UMLDeformerModelInstance::CheckCompatibility(USkeletalMeshComponent* InSkelMeshComponent, bool LogIssues)
+FString UMLDeformerModelInstance::CheckCompatibility(USkeletalMeshComponent* InSkelMeshComponent, bool bLogIssues)
 {
 	ErrorText = FString();
 
@@ -137,7 +128,7 @@ FString UMLDeformerModelInstance::CheckCompatibility(USkeletalMeshComponent* InS
 	{
 		ErrorText += InputInfo->GenerateCompatibilityErrorString(SkelMesh);
 		ErrorText += "\n";
-		if (LogIssues)
+		if (bLogIssues)
 		{
 			UE_LOG(LogMLDeformer, Error, TEXT("ML Deformer '%s' isn't compatible with Skeletal Mesh '%s'.\nReason(s):\n%s"), 
 				*Model->GetDeformerAsset()->GetName(), 
@@ -146,23 +137,6 @@ FString UMLDeformerModelInstance::CheckCompatibility(USkeletalMeshComponent* InS
 		}
 	}
 
-	// Verify the number of inputs versus the expected number of inputs.
-	UNeuralNetwork* NeuralNetwork = Model->GetNeuralNetwork();
-	if (NeuralNetwork && NeuralNetwork->IsLoaded() && Model->GetDeformerAsset())
-	{
-		const int64 NumNeuralNetInputs = NeuralNetwork->GetInputTensor().Num();
-		const int64 NumDeformerAssetInputs = static_cast<int64>(Model->GetInputInfo()->CalcNumNeuralNetInputs(Model->GetNumFloatsPerBone(), Model->GetNumFloatsPerCurve()));
-		if (NumNeuralNetInputs != NumDeformerAssetInputs)
-		{
-			const FString InputErrorString = "The number of network inputs doesn't match the asset. Please retrain the asset."; 
-			ErrorText += InputErrorString + "\n";
-			if (LogIssues)
-			{
-				UE_LOG(LogMLDeformer, Error, TEXT("Deformer '%s': %s"), *(Model->GetDeformerAsset()->GetName()), *InputErrorString);
-			}
-		}
-	}
-
 	return ErrorText;
 }
 
@@ -273,106 +247,6 @@ int64 UMLDeformerModelInstance::SetNeuralNetworkInputValues(float* InputData, in
 
 bool UMLDeformerModelInstance::IsValidForDataProvider() const
 {
-	const UNeuralNetwork* NeuralNetwork = Model->GetNeuralNetwork();
-	if (NeuralNetwork == nullptr || !NeuralNetwork->IsLoaded())
-	{
-		return false;
-	}
-
-	// We expect to run on the GPU when using a data provider for the deformer graph system (Optimus).
-	if (Model->IsNeuralNetworkOnGPU())
-	{
-		// Make sure we're actually running the network on the GPU.
-		// Inputs are expected to come from the CPU though.
-		if (NeuralNetwork->GetDeviceType() != ENeuralDeviceType::GPU || NeuralNetwork->GetOutputDeviceType() != ENeuralDeviceType::GPU || NeuralNetwork->GetInputDeviceType() != ENeuralDeviceType::CPU)
-		{
-			return false;
-		}
-	}
-
-	return (Model->GetVertexMapBuffer().ShaderResourceViewRHI != nullptr) && (GetNeuralNetworkInferenceHandle() != -1);
-}
-
-void UMLDeformerModelInstance::Execute(float ModelWeight)
-{
-	TRACE_CPUPROFILER_EVENT_SCOPE(UMLDeformerModelInstance::Execute)
-
-	UNeuralNetwork* NeuralNetwork = Model->GetNeuralNetwork();
-	if (NeuralNetwork == nullptr)
-	{
-		return;
-	}
-
-	if (Model->IsNeuralNetworkOnGPU())
-	{
-		// Even if the model needs the GPU it is possible that the hardware does not support GPU evaluation
-		if (NeuralNetwork->GetDeviceType() == ENeuralDeviceType::GPU)
-		{
-			// NOTE: Inputs still come from the CPU.
-			check(NeuralNetwork->GetDeviceType() == ENeuralDeviceType::GPU && NeuralNetwork->GetInputDeviceType() == ENeuralDeviceType::CPU && NeuralNetwork->GetOutputDeviceType() == ENeuralDeviceType::GPU);
-			ENQUEUE_RENDER_COMMAND(RunNeuralNetwork)
-				(
-					[NeuralNetwork, Handle = NeuralNetworkInferenceHandle](FRHICommandListImmediate& RHICmdList)
-					{
-						// Output deltas will be available on GPU for DeformerGraph via UMLDeformerDataProvider.
-						FRDGBuilder GraphBuilder(RHICmdList);
-						NeuralNetwork->Run(GraphBuilder, Handle);
-						GraphBuilder.Execute();
-					}
-			);
-		}
-	}
-	else
-	{
-		// Run on the CPU.
-		check(NeuralNetwork->GetDeviceType() == ENeuralDeviceType::CPU && NeuralNetwork->GetInputDeviceType() == ENeuralDeviceType::CPU && NeuralNetwork->GetOutputDeviceType() == ENeuralDeviceType::CPU);
-		NeuralNetwork->Run(NeuralNetworkInferenceHandle);
-	}
-}
-
-bool UMLDeformerModelInstance::SetupInputs()
-{
-	TRACE_CPUPROFILER_EVENT_SCOPE(UMLDeformerModelInstance::SetupInputs)
-
-	// Some safety checks.
-	if (Model == nullptr ||
-		SkeletalMeshComponent == nullptr ||
-		SkeletalMeshComponent->GetSkeletalMeshAsset() == nullptr ||
-		!bIsCompatible)
-	{
-		return false;
-	}
-
-	// Get the network and make sure it's loaded.
-	UNeuralNetwork* NeuralNetwork = Model->GetNeuralNetwork();
-	if (NeuralNetwork == nullptr || !NeuralNetwork->IsLoaded())
-	{
-		return false;
-	}
-
-	// Allocate an inference context if none has already been allocated.
-	if (NeuralNetworkInferenceHandle == -1)
-	{
-		NeuralNetworkInferenceHandle = NeuralNetwork->CreateInferenceContext();
-		if (NeuralNetworkInferenceHandle == -1)
-		{
-			return false;
-		}
-	}
-
-	// If the neural network expects a different number of inputs, do nothing.
-	const int64 NumNeuralNetInputs = NeuralNetwork->GetInputTensorForContext(NeuralNetworkInferenceHandle).Num();
-	const int64 NumDeformerAssetInputs = Model->GetInputInfo()->CalcNumNeuralNetInputs(Model->GetNumFloatsPerBone(), Model->GetNumFloatsPerCurve());
-	if (NumNeuralNetInputs != NumDeformerAssetInputs)
-	{
-		return false;
-	}
-
-	// Update and write the input values directly into the input tensor.
-	float* InputDataPointer = static_cast<float*>(NeuralNetwork->GetInputDataPointerMutableForContext(NeuralNetworkInferenceHandle));
-	const int64 NumFloatsWritten = SetNeuralNetworkInputValues(InputDataPointer, NumNeuralNetInputs);
-	check(NumFloatsWritten == NumNeuralNetInputs);
-
 	return true;
 }
 
diff --git a/Engine/Plugins/Animation/MLDeformer/MLDeformerFramework/Source/MLDeformerFramework/Private/MLDeformerMorphModelInstance.cpp b/Engine/Plugins/Animation/MLDeformer/MLDeformerFramework/Source/MLDeformerFramework/Private/MLDeformerMorphModelInstance.cpp
--- a/Engine/Plugins/Animation/MLDeformer/MLDeformerFramework/Source/MLDeformerFramework/Private/MLDeformerMorphModelInstance.cpp
+++ b/Engine/Plugins/Animation/MLDeformer/MLDeformerFramework/Source/MLDeformerFramework/Private/MLDeformerMorphModelInstance.cpp
@@ -3,7 +3,6 @@
 #include "MLDeformerMorphModelInstance.h"
 #include "MLDeformerMorphModel.h" 
 #include "MLDeformerComponent.h" 
-#include "NeuralNetwork.h"
 #include "Components/ExternalMorphSet.h"
 #include "Components/SkeletalMeshComponent.h"
 
@@ -159,75 +158,3 @@ void UMLDeformerMorphModelInstance::HandleZeroModelWeight()
 		WeightData->ZeroWeights();
 	}
 }
-
-// Run the neural network, which calculates its outputs, which are the weights of our morph targets.
-void UMLDeformerMorphModelInstance::Execute(float ModelWeight)
-{
-	TRACE_CPUPROFILER_EVENT_SCOPE(UMLDeformerMorphModelInstance::Execute)
-
-	// Grab the weight data for this morph set.
-	// This could potentially fail if we are applying this deformer to the wrong skeletal mesh component.
-	const int LOD = 0;	// For now we only support LOD 0, as we can't setup an ML Deformer per LOD yet.
-	FExternalMorphSetWeights* WeightData = FindWeightData(LOD);
-	if (WeightData == nullptr)
-	{
-		return;
-	}
-
-	// If our model is active, we want to run the neural network and update the morph weights
-	// with the values that the neural net calculated for us.
-	const UNeuralNetwork* NeuralNetwork = Model->GetNeuralNetwork();
-	if (NeuralNetwork)
-	{
-		// Perform the neural network inference, which updates the output tensor.
-		// This takes most CPU time inside this method.
-		UMLDeformerModelInstance::Execute(ModelWeight);
-
-		// Get the output tensor, read the values and use them as morph target weights inside the skeletal mesh component.
-		const FNeuralTensor& OutputTensor = NeuralNetwork->GetOutputTensorForContext(NeuralNetworkInferenceHandle);
-		const int32 NumNetworkWeights = OutputTensor.Num();
-		const int32 NumMorphTargets = WeightData->Weights.Num();;
-
-		// If we have the expected amount of weights.
-		// +1 because we always have an extra morph target that represents the means, with fixed weight of 1.
-		// Therefore, the neural network output tensor will contain one less float than the number of morph targets in our morph set.
-		if (NumMorphTargets == NumNetworkWeights + 1)
-		{
-			// Set the first morph target, which represents the means, to a weight of 1.0, as it always needs to be fully active.
-			WeightData->Weights[0] = ModelWeight;
-
-			// Update all generated morph target weights with the values calculated by our neural network.
-			UMLDeformerMorphModel* MorphModel = Cast<UMLDeformerMorphModel>(Model);
-			const TArrayView<const float> ErrorValues = MorphModel->GetMorphTargetErrorValues();
-			const TArrayView<const int32> ErrorOrder = MorphModel->GetMorphTargetErrorOrder();
-			if (!ErrorValues.IsEmpty())
-			{
-				const int32 QualityLevel = GetMLDeformerComponent()->GetQualityLevel();
-				const int32 NumActiveMorphs = MorphModel->GetNumActiveMorphs(QualityLevel);
-				for (int32 Index = 0; Index < NumActiveMorphs; ++Index)
-				{
-					const int32 MorphIndex = ErrorOrder[Index];
-					const float TargetWeight = OutputTensor.At<float>(MorphIndex) * ModelWeight;
-					WeightData->Weights[MorphIndex + 1] = FMath::Lerp<float, float>(StartMorphWeights[MorphIndex + 1], TargetWeight, MorphLerpAlpha);
-				}
-
-				// Disable all inactive morphs.
-				for (int32 Index = NumActiveMorphs; Index < NumNetworkWeights; ++Index)
-				{
-					const int32 MorphIndex = ErrorOrder[Index];
-					WeightData->Weights[MorphIndex + 1] = FMath::Lerp<float, float>(StartMorphWeights[MorphIndex + 1], 0.0f, MorphLerpAlpha);
-				}
-			}
-			else
-			{
-				for (int32 MorphIndex = 0; MorphIndex < NumNetworkWeights; ++MorphIndex)
-				{
-					WeightData->Weights[MorphIndex + 1] = OutputTensor.At<float>(MorphIndex) * ModelWeight;
-				}
-			}
-			return;
-		}
-	}
-
-	WeightData->ZeroWeights();
-}
diff --git a/Engine/Plugins/Animation/MLDeformer/MLDeformerFramework/Source/MLDeformerFramework/Public/MLDeformerGraphDataInterface.h b/Engine/Plugins/Animation/MLDeformer/MLDeformerFramework/Source/MLDeformerFramework/Public/MLDeformerGraphDataInterface.h
deleted file mode 100644
diff --git a/Engine/Plugins/Animation/MLDeformer/MLDeformerFramework/Source/MLDeformerFramework/Public/MLDeformerGraphDebugDataInterface.h b/Engine/Plugins/Animation/MLDeformer/MLDeformerFramework/Source/MLDeformerFramework/Public/MLDeformerGraphDebugDataInterface.h
deleted file mode 100644
diff --git a/Engine/Plugins/Animation/MLDeformer/MLDeformerFramework/Source/MLDeformerFramework/Public/MLDeformerModel.h b/Engine/Plugins/Animation/MLDeformer/MLDeformerFramework/Source/MLDeformerFramework/Public/MLDeformerModel.h
--- a/Engine/Plugins/Animation/MLDeformer/MLDeformerFramework/Source/MLDeformerFramework/Public/MLDeformerModel.h
+++ b/Engine/Plugins/Animation/MLDeformer/MLDeformerFramework/Source/MLDeformerFramework/Public/MLDeformerModel.h
@@ -10,7 +10,6 @@
 #include "BoneContainer.h"
 #include "RenderCommandFence.h"
 #include "RenderResource.h"
-#include "NeuralNetwork.h"
 #include "Animation/AnimSequence.h"
 #include "MLDeformerCurveReference.h"
 #include "MLDeformerModel.generated.h"
@@ -20,6 +19,7 @@ class UMLDeformerVizSettings;
 class UMLDeformerModelInstance;
 class UMLDeformerComponent;
 class UMLDeformerInputInfo;
+class UNeuralNetwork;
 
 /** The channel to get the mask data from. */
 UENUM()
@@ -186,6 +186,12 @@ public:
 	 */
 	void SetSkeletalMesh(USkeletalMesh* SkelMesh)			{ SkeletalMesh = SkelMesh; }
 
+	UE_DEPRECATED(5.2, "This method will be deleted. We are moving neural networks into the derived models instead. This method will not do anything.")
+	virtual void SetNeuralNetwork(UNeuralNetwork* InNeuralNetwork) {}
+
+	UE_DEPRECATED(5.2, "This method will be deleted. We are moving neural networks into the derived models instead. A nullptr will be returned.")
+	UNeuralNetwork* GetNeuralNetwork() const				{ return nullptr; }
+
 #if WITH_EDITORONLY_DATA
 	/**
 	 * Check whether this model currently has a training mesh setup or not.
@@ -315,19 +321,6 @@ public:
 	 */
 	const UE::MLDeformer::FVertexMapBuffer& GetVertexMapBuffer() const { return VertexMapBuffer; }
 
-	/**
-	 * Get the neural network that we have trained. This can return a nullptr when no network has been trained yet.
-	 * This network is used during inference.
-	 * @return A pointer to the neural network, or nullptr when the network has not yet been trained.
-	 */
-	UNeuralNetwork* GetNeuralNetwork() const					{ return NeuralNetwork.Get(); }
-
-	/**
-	 * Set the neural network object that we use during inference.
-	 * @param InNeuralNetwork The new neural network to use inside this deformer model.
-	 */
-	virtual void SetNeuralNetwork(UNeuralNetwork* InNeuralNetwork);
-
 	/**
 	 * Get the neural network modified delegate.
 	 * This triggers when the neural network pointers changes.
@@ -594,10 +587,6 @@ private:
 	UPROPERTY()
 	TArray<int32> VertexMap;
 
-	/** The neural network that is used during inference. */
-	UPROPERTY()
-	TObjectPtr<UNeuralNetwork> NeuralNetwork = nullptr;
-
 	/** The skeletal mesh that represents the linear skinned mesh. */
 	UPROPERTY(EditAnywhere, Category = "Base Mesh")
 	TObjectPtr<USkeletalMesh> SkeletalMesh = nullptr;
diff --git a/Engine/Plugins/Animation/MLDeformer/MLDeformerFramework/Source/MLDeformerFramework/Public/MLDeformerModelInstance.h b/Engine/Plugins/Animation/MLDeformer/MLDeformerFramework/Source/MLDeformerFramework/Public/MLDeformerModelInstance.h
--- a/Engine/Plugins/Animation/MLDeformer/MLDeformerFramework/Source/MLDeformerFramework/Public/MLDeformerModelInstance.h
+++ b/Engine/Plugins/Animation/MLDeformer/MLDeformerFramework/Source/MLDeformerFramework/Public/MLDeformerModelInstance.h
@@ -9,7 +9,6 @@
 #include "MLDeformerModelInstance.generated.h"
 
 class UMLDeformerModel;
-class UNeuralNetwork;
 class USkeletalMeshComponent;
 class UMLDeformerComponent;
 
@@ -67,11 +66,11 @@ public:
 	 * Check whether the deformer is compatible with a given skeletal mesh component.
 	 * This internally also edits the value returned by GetCompatibilityErrorText().
 	 * @param InSkelMeshComponent The skeletal mesh component to check compatibility with.
-	 * @param LogIssues Set to true to automatically log any compatibility errors.
+	 * @param bLogIssues Set to true to automatically log any compatibility errors.
 	 * @return Returns the error string. When the returned string is empty, there were no errors and thus
 	 *         the specified skeletal mesh component is compatible.
 	 */
-	virtual FString CheckCompatibility(USkeletalMeshComponent* InSkelMeshComponent, bool LogIssues=false);
+	virtual FString CheckCompatibility(USkeletalMeshComponent* InSkelMeshComponent, bool bLogIssues=false);
 
 	/**
 	 * Check if we are in a valid state for the deformer graph's data provider.
@@ -146,15 +145,14 @@ protected:
 	 * @return Returns true when the setup is done correctly and the neural network is ready to be executed. Otherwise false is returned, which 
 	 * can happen when the NeuralNetwork pointer is invalid, when the model is not set, when the network is not compatible, etc.
 	 */
-	virtual bool SetupInputs();
+	virtual bool SetupInputs() { return true; }
 
 	/**
 	 * Execute the model instance, which can run the neural network inference in case the model uses a neural network.
 	 * This already assumes that compatibility checks are done, and that the network inputs are set etc.
-	 * Internally this will typically execute the UNeuralNetwork::Run() method, either on GPU or CPU.
 	 * @param ModelWeight The weight of the model, must be between 0 and 1.
 	 */
-	virtual void Execute(float ModelWeight);
+	virtual void Execute(float ModelWeight) {}
 
 	/**
 	 * Handle when the model weight is zero.
diff --git a/Engine/Plugins/Animation/MLDeformer/MLDeformerFramework/Source/MLDeformerFramework/Public/MLDeformerMorphModelInstance.h b/Engine/Plugins/Animation/MLDeformer/MLDeformerFramework/Source/MLDeformerFramework/Public/MLDeformerMorphModelInstance.h
--- a/Engine/Plugins/Animation/MLDeformer/MLDeformerFramework/Source/MLDeformerFramework/Public/MLDeformerMorphModelInstance.h
+++ b/Engine/Plugins/Animation/MLDeformer/MLDeformerFramework/Source/MLDeformerFramework/Public/MLDeformerMorphModelInstance.h
@@ -28,7 +28,6 @@ public:
 	virtual void PostMLDeformerComponentInit() override;
 	virtual void Release() override;
 	virtual void Tick(float DeltaTime, float ModelWeight) override;
-	virtual void Execute(float ModelWeight) override;
 	virtual void HandleZeroModelWeight() override;
 	virtual bool IsValidForDataProvider() const override;
 	// ~END UMLDeformerModelInstance overrides.
diff --git a/Engine/Plugins/Animation/MLDeformer/MLDeformerFramework/Source/MLDeformerFrameworkEditor/MLDeformerFrameworkEditor.Build.cs b/Engine/Plugins/Animation/MLDeformer/MLDeformerFramework/Source/MLDeformerFrameworkEditor/MLDeformerFrameworkEditor.Build.cs
--- a/Engine/Plugins/Animation/MLDeformer/MLDeformerFramework/Source/MLDeformerFrameworkEditor/MLDeformerFrameworkEditor.Build.cs
+++ b/Engine/Plugins/Animation/MLDeformer/MLDeformerFramework/Source/MLDeformerFrameworkEditor/MLDeformerFrameworkEditor.Build.cs
@@ -45,7 +45,6 @@ namespace UnrealBuildTool.Rules
 					"GeometryCache",
 					"SequencerWidgets",
 					"TimeManagement",
-					"NeuralNetworkInference",
 					"RenderCore",
 					"RHI"
 				}
diff --git a/Engine/Plugins/Animation/MLDeformer/MLDeformerFramework/Source/MLDeformerFrameworkEditor/Private/MLDeformerEditorModel.cpp b/Engine/Plugins/Animation/MLDeformer/MLDeformerFramework/Source/MLDeformerFrameworkEditor/Private/MLDeformerEditorModel.cpp
--- a/Engine/Plugins/Animation/MLDeformer/MLDeformerFramework/Source/MLDeformerFrameworkEditor/Private/MLDeformerEditorModel.cpp
+++ b/Engine/Plugins/Animation/MLDeformer/MLDeformerFramework/Source/MLDeformerFrameworkEditor/Private/MLDeformerEditorModel.cpp
@@ -34,7 +34,6 @@
 #include "GeometryCache.h"
 #include "GeometryCacheMeshData.h"
 #include "GeometryCacheTrack.h"
-#include "NeuralNetwork.h"
 #include "BoneContainer.h"
 
 #define LOCTEXT_NAMESPACE "MLDeformerEditorModel"
@@ -1526,43 +1525,6 @@ namespace UE::MLDeformer
 		return nullptr;
 	}
 
-	UNeuralNetwork* FMLDeformerEditorModel::LoadNeuralNetworkFromOnnx(const FString& Filename) const
-	{
-		const FString OnnxFile = FPaths::ConvertRelativePathToFull(Filename);
-		if (FPaths::FileExists(OnnxFile))
-		{
-			UE_LOG(LogMLDeformer, Display, TEXT("Loading Onnx file '%s'..."), *OnnxFile);
-			UNeuralNetwork* Result = NewObject<UNeuralNetwork>(Model, UNeuralNetwork::StaticClass());		
-			if (Result->Load(OnnxFile))
-			{
-				if (Model->IsNeuralNetworkOnGPU())
-				{
-					Result->SetDeviceType(ENeuralDeviceType::GPU, ENeuralDeviceType::CPU, ENeuralDeviceType::GPU);	
-					if (Result->GetDeviceType() != ENeuralDeviceType::GPU || Result->GetOutputDeviceType() != ENeuralDeviceType::GPU || Result->GetInputDeviceType() != ENeuralDeviceType::CPU)
-					{
-						UE_LOG(LogMLDeformer, Error, TEXT("Neural net in ML Deformer '%s' cannot run on the GPU, it will not be active."), *Model->GetDeformerAsset()->GetName());
-					}
-				}
-				else
-				{
-					Result->SetDeviceType(ENeuralDeviceType::CPU, ENeuralDeviceType::CPU, ENeuralDeviceType::CPU);
-				}
-				UE_LOG(LogMLDeformer, Display, TEXT("Successfully loaded Onnx file '%s'..."), *OnnxFile);
-				return Result;
-			}
-			else
-			{
-				UE_LOG(LogMLDeformer, Error, TEXT("Failed to load Onnx file '%s'"), *OnnxFile);
-			}
-		}
-		else
-		{
-			UE_LOG(LogMLDeformer, Error, TEXT("Onnx file '%s' does not exist!"), *OnnxFile);
-		}
-
-		return nullptr;
-	}
-
 	bool FMLDeformerEditorModel::IsEditorReadyForTrainingBasicChecks()
 	{
 		// Make sure we have picked required assets.
@@ -1584,24 +1546,6 @@ namespace UE::MLDeformer
 		return true;
 	}
 
-	bool FMLDeformerEditorModel::LoadTrainedNetwork() const
-	{
-		const FString OnnxFile = GetTrainedNetworkOnnxFile();
-		UNeuralNetwork* Network = LoadNeuralNetworkFromOnnx(OnnxFile);
-		if (Network)
-		{
-			Model->SetNeuralNetwork(Network);
-			return true;
-		}
-
-		return false;
-	}
-
-	bool FMLDeformerEditorModel::IsTrained() const
-	{
-		return (Model->GetNeuralNetwork() != nullptr);
-	}
-
 	void FMLDeformerEditorModel::UpdateMemoryUsage()
 	{
 		Model->InvalidateMemUsage();
diff --git a/Engine/Plugins/Animation/MLDeformer/MLDeformerFramework/Source/MLDeformerFrameworkEditor/Private/MLDeformerModelDetails.cpp b/Engine/Plugins/Animation/MLDeformer/MLDeformerFramework/Source/MLDeformerFrameworkEditor/Private/MLDeformerModelDetails.cpp
--- a/Engine/Plugins/Animation/MLDeformer/MLDeformerFramework/Source/MLDeformerFrameworkEditor/Private/MLDeformerModelDetails.cpp
+++ b/Engine/Plugins/Animation/MLDeformer/MLDeformerFramework/Source/MLDeformerFrameworkEditor/Private/MLDeformerModelDetails.cpp
@@ -9,7 +9,6 @@
 #include "MLDeformerEditorModel.h"
 #include "MLDeformerEditorModule.h"
 #include "MLDeformerGeomCacheHelpers.h"
-#include "NeuralNetwork.h"
 #include "GeometryCache.h"
 #include "GeometryCacheTrack.h"
 #include "Animation/AnimSequence.h"
diff --git a/Engine/Plugins/Animation/MLDeformer/MLDeformerFramework/Source/MLDeformerFrameworkEditor/Public/MLDeformerEditorModel.h b/Engine/Plugins/Animation/MLDeformer/MLDeformerFramework/Source/MLDeformerFrameworkEditor/Public/MLDeformerEditorModel.h
--- a/Engine/Plugins/Animation/MLDeformer/MLDeformerFramework/Source/MLDeformerFrameworkEditor/Public/MLDeformerEditorModel.h
+++ b/Engine/Plugins/Animation/MLDeformer/MLDeformerFramework/Source/MLDeformerFrameworkEditor/Public/MLDeformerEditorModel.h
@@ -513,6 +513,9 @@ namespace UE::MLDeformer
 		 */
 		virtual void UpdateDeformerGraph();
 
+		UE_DEPRECATED(5.2, "This method will be deleted and will just return a nullptr right now.")
+		UNeuralNetwork* LoadNeuralNetworkFromOnnx(const FString& Filename) const { return nullptr; }
+
 		/**
 		 * Sample the vertex deltas between the training base model and target model.
 		 * This will initialize the sampler if needed and set the sampler space to post-skinning deltas, and then samples them using the sampler.
@@ -525,7 +528,7 @@ namespace UE::MLDeformer
 		 * Load the trained network from an onnx file.
 		 * The filename of the onnx file is determined by the GetTrainedNetworkOnnxFile method.
 		 */
-		virtual bool LoadTrainedNetwork() const;
+		virtual bool LoadTrainedNetwork() const { return false; }
 
 		/**
 		 * Get the onnx file name of the trained network.
@@ -541,7 +544,7 @@ namespace UE::MLDeformer
 		 * This is determined by looking whether the neural network pointer is nullptr or not.
 		 * @return Returns true when the model is trained already, or false if it hasn't been trained yet.
 		 */
-		virtual bool IsTrained() const;
+		virtual bool IsTrained() const { return false; }
 
 		/**
 		 * Get the editor actor that defines the timeline play position.
@@ -744,13 +747,6 @@ namespace UE::MLDeformer
 		 */
 		int32 GetNumCurvesOnSkeletalMesh(USkeletalMesh* SkelMesh) const;
 
-		/**
-		 * Load a neural network from a given Onnx file.
-		 * @param Filename The onnx file name.
-		 * @return A pointer to the newly created neural network, or nullptr in case it failed to load.
-		 */
-		UNeuralNetwork* LoadNeuralNetworkFromOnnx(const FString& Filename) const;
-
 		/**
 		 * Get the currently desired training frame number.
 		 * You can call CheckTrainingFrameChanged in order to make the desired frame the actual frame.
diff --git a/Engine/Plugins/Animation/MLDeformer/NearestNeighborModel/NearestNeighborModel.uplugin b/Engine/Plugins/Animation/MLDeformer/NearestNeighborModel/NearestNeighborModel.uplugin
--- a/Engine/Plugins/Animation/MLDeformer/NearestNeighborModel/NearestNeighborModel.uplugin
+++ b/Engine/Plugins/Animation/MLDeformer/NearestNeighborModel/NearestNeighborModel.uplugin
@@ -1,7 +1,7 @@
 {
 	"FileVersion": 3,
 	"Version": 1,
-	"VersionName": "0.1",
+	"VersionName": "0.2",
 	"FriendlyName": "ML Deformer Nearest Neighbor Model",
 	"Description": "Nearest Neighbor Model for the ML Deformer Framework",
 	"Category": "Animation",
diff --git a/Engine/Plugins/Animation/MLDeformer/NearestNeighborModel/Source/NearestNeighborModel/Private/NearestNeighborModel.cpp b/Engine/Plugins/Animation/MLDeformer/NearestNeighborModel/Source/NearestNeighborModel/Private/NearestNeighborModel.cpp
--- a/Engine/Plugins/Animation/MLDeformer/NearestNeighborModel/Source/NearestNeighborModel/Private/NearestNeighborModel.cpp
+++ b/Engine/Plugins/Animation/MLDeformer/NearestNeighborModel/Source/NearestNeighborModel/Private/NearestNeighborModel.cpp
@@ -71,9 +71,9 @@ void UNearestNeighborModel::Serialize(FArchive& Archive)
 	{
 		if (DoesUseOptimizedNetwork())
 		{
-			if (GetNeuralNetwork() != nullptr)
+			if (GetNNINetwork() != nullptr)
 			{
-				SetNeuralNetwork(nullptr);
+				SetNNINetwork(nullptr);
 			}
 		}
 		else
@@ -100,6 +100,16 @@ void UNearestNeighborModel::PostLoad()
 #endif
 }
 
+void UNearestNeighborModel::SetNNINetwork(UNeuralNetwork* InNeuralNetwork)
+{
+	GetNeuralNetworkModifyDelegate().Broadcast();
+	NNINetwork = InNeuralNetwork;
+}
+
+UNeuralNetwork* UNearestNeighborModel::GetNNINetwork() const
+{ 
+	return NNINetwork.Get();
+}
 
 TArray<uint32> ReadTxt(const FString &Path)
 {
@@ -400,7 +410,7 @@ int32 UNearestNeighborModel::GetNumNeighborsFromAnimSequence(int32 PartId) const
 
 void UNearestNeighborModel::UpdateNetworkSize()
 {
-	UNeuralNetwork* Network = GetNeuralNetwork();
+	UNeuralNetwork* Network = GetNNINetwork();
 	if (Network != nullptr)
 	{
 		const SIZE_T NumBytes = Network->GetResourceSizeBytes(EResourceSizeMode::EstimatedTotal);
@@ -510,7 +520,7 @@ bool UNearestNeighborModel::LoadOptimizedNetwork(const FString& OnnxPath)
 	}
 	else
 	{
-		UE_LOG(LogMLDeformer, Error, TEXT("Onnx file '%s' does not exist!"), *OnnxFile);
+		UE_LOG(LogNearestNeighborModel, Error, TEXT("Onnx file '%s' does not exist!"), *OnnxFile);
 	}
 
 	return false;
diff --git a/Engine/Plugins/Animation/MLDeformer/NearestNeighborModel/Source/NearestNeighborModel/Private/NearestNeighborModelInstance.cpp b/Engine/Plugins/Animation/MLDeformer/NearestNeighborModel/Source/NearestNeighborModel/Private/NearestNeighborModelInstance.cpp
--- a/Engine/Plugins/Animation/MLDeformer/NearestNeighborModel/Source/NearestNeighborModel/Private/NearestNeighborModelInstance.cpp
+++ b/Engine/Plugins/Animation/MLDeformer/NearestNeighborModel/Source/NearestNeighborModel/Private/NearestNeighborModelInstance.cpp
@@ -277,7 +277,7 @@ void UNearestNeighborModelInstance::GetInputDataPointer(float*& OutInputData, in
 	}
 	else
 	{
-		UNeuralNetwork* NeuralNetwork = NearestNeighborModel->GetNeuralNetwork();
+		UNeuralNetwork* NeuralNetwork = NearestNeighborModel->GetNNINetwork();
 		if (NeuralNetwork)
 		{
 			OutInputData = static_cast<float*>(NeuralNetwork->GetInputDataPointerMutableForContext(NeuralNetworkInferenceHandle));
@@ -305,7 +305,7 @@ void UNearestNeighborModelInstance::GetOutputDataPointer(float*& OutOutputData,
 	}
 	else
 	{
-		UNeuralNetwork* NeuralNetwork = NearestNeighborModel->GetNeuralNetwork();
+		UNeuralNetwork* NeuralNetwork = NearestNeighborModel->GetNNINetwork();
 		if (NeuralNetwork)
 		{
 			const FNeuralTensor& OutputTensor = NeuralNetwork->GetOutputTensorForContext(NeuralNetworkInferenceHandle);
@@ -331,7 +331,7 @@ FString UNearestNeighborModelInstance::CheckCompatibility(USkeletalMeshComponent
 		ErrorText += "\n";
 		if (LogIssues)
 		{
-			UE_LOG(LogMLDeformer, Error, TEXT("ML Deformer '%s' isn't compatible with Skeletal Mesh '%s'.\nReason(s):\n%s"), 
+			UE_LOG(LogNearestNeighborModel, Error, TEXT("ML Deformer '%s' isn't compatible with Skeletal Mesh '%s'.\nReason(s):\n%s"), 
 				*Model->GetDeformerAsset()->GetName(), 
 				*SkelMesh->GetName(), 
 				*ErrorText);
@@ -356,7 +356,7 @@ FString UNearestNeighborModelInstance::CheckCompatibility(USkeletalMeshComponent
 	}
 	else
 	{
-		UNeuralNetwork* NeuralNetwork = Model->GetNeuralNetwork();
+		UNeuralNetwork* NeuralNetwork = NearestNeighborModel->GetNNINetwork();
 		if (NeuralNetwork && NeuralNetwork->IsLoaded())
 		{
 			NumNeuralNetInputs = NeuralNetwork->GetInputTensor().Num();
@@ -370,7 +370,7 @@ FString UNearestNeighborModelInstance::CheckCompatibility(USkeletalMeshComponent
 		ErrorText += InputErrorString + "\n";
 		if (LogIssues && Model->GetDeformerAsset())
 		{
-				UE_LOG(LogMLDeformer, Error, TEXT("Deformer '%s': %s"), *(Model->GetDeformerAsset()->GetName()), *InputErrorString);
+			UE_LOG(LogNearestNeighborModel, Error, TEXT("Deformer '%s': %s"), *(Model->GetDeformerAsset()->GetName()), *InputErrorString);
 		}
 	}
 
diff --git a/Engine/Plugins/Animation/MLDeformer/NearestNeighborModel/Source/NearestNeighborModel/Public/NearestNeighborModel.h b/Engine/Plugins/Animation/MLDeformer/NearestNeighborModel/Source/NearestNeighborModel/Public/NearestNeighborModel.h
--- a/Engine/Plugins/Animation/MLDeformer/NearestNeighborModel/Source/NearestNeighborModel/Public/NearestNeighborModel.h
+++ b/Engine/Plugins/Animation/MLDeformer/NearestNeighborModel/Source/NearestNeighborModel/Public/NearestNeighborModel.h
@@ -33,7 +33,7 @@ namespace UE::NearestNeighborModel
 	private:
 		void InitRHI() override;
 
-		TArray<T> *Array = nullptr;
+		TArray<T>* Array = nullptr;
 		FString DebugName;
 	};
 
@@ -244,6 +244,9 @@ public:
 	void SetOptimizedNetwork(UNearestNeighborOptimizedNetwork* InOptimizedNetwork);
 	bool LoadOptimizedNetwork(const FString& OnnxPath);
 
+	UNeuralNetwork* GetNNINetwork() const;
+	void SetNNINetwork(UNeuralNetwork* InNeuralNetwork);
+
 #if WITH_EDITORONLY_DATA
 	TObjectPtr<UAnimSequence> GetNearestNeighborSkeletons(int32 PartId);
 	const TObjectPtr<UAnimSequence> GetNearestNeighborSkeletons(int32 PartId) const;
@@ -423,4 +426,8 @@ public:
 
 	UPROPERTY()
 	TObjectPtr<UNearestNeighborOptimizedNetwork> OptimizedNetwork = nullptr;
+
+	/** The NNI neural network. */
+	UPROPERTY()
+	TObjectPtr<UNeuralNetwork> NNINetwork;
 };
diff --git a/Engine/Plugins/Animation/MLDeformer/NearestNeighborModel/Source/NearestNeighborModelEditor/Private/NearestNeighborEditorModel.cpp b/Engine/Plugins/Animation/MLDeformer/NearestNeighborModel/Source/NearestNeighborModelEditor/Private/NearestNeighborEditorModel.cpp
--- a/Engine/Plugins/Animation/MLDeformer/NearestNeighborModel/Source/NearestNeighborModelEditor/Private/NearestNeighborEditorModel.cpp
+++ b/Engine/Plugins/Animation/MLDeformer/NearestNeighborModel/Source/NearestNeighborModelEditor/Private/NearestNeighborEditorModel.cpp
@@ -90,6 +90,36 @@ namespace UE::NearestNeighborModel
 		}
 	}
 
+	UNeuralNetwork* FNearestNeighborEditorModel::LoadNeuralNetworkFromOnnx(const FString& Filename) const
+	{
+		const FString OnnxFile = FPaths::ConvertRelativePathToFull(Filename);
+		if (FPaths::FileExists(OnnxFile))
+		{
+			UE_LOG(LogNearestNeighborModel, Display, TEXT("Loading Onnx file '%s'..."), *OnnxFile);
+			UNeuralNetwork* Result = NewObject<UNeuralNetwork>(Model, UNeuralNetwork::StaticClass());		
+			if (Result->Load(OnnxFile))
+			{
+				Result->SetDeviceType(ENeuralDeviceType::GPU, ENeuralDeviceType::CPU, ENeuralDeviceType::GPU);	
+				if (Result->GetDeviceType() != ENeuralDeviceType::GPU || Result->GetOutputDeviceType() != ENeuralDeviceType::GPU || Result->GetInputDeviceType() != ENeuralDeviceType::CPU)
+				{
+					UE_LOG(LogNearestNeighborModel, Error, TEXT("Neural net in ML Deformer '%s' cannot run on the GPU, it will not be active."), *Model->GetDeformerAsset()->GetName());
+				}
+				UE_LOG(LogNearestNeighborModel, Display, TEXT("Successfully loaded Onnx file '%s'..."), *OnnxFile);
+				return Result;
+			}
+			else
+			{
+				UE_LOG(LogNearestNeighborModel, Error, TEXT("Failed to load Onnx file '%s'"), *OnnxFile);
+			}
+		}
+		else
+		{
+			UE_LOG(LogNearestNeighborModel, Error, TEXT("Onnx file '%s' does not exist!"), *OnnxFile);
+		}
+
+		return nullptr;
+	}
+
 	bool FNearestNeighborEditorModel::LoadTrainedNetwork() const
 	{
 		UNearestNeighborModel* NearestNeighborModel = GetNearestNeighborModel();
@@ -101,7 +131,7 @@ namespace UE::NearestNeighborModel
 			{
 				if (NearestNeighborModel->DoesUseOptimizedNetwork())
 				{
-					Model->SetNeuralNetwork(Network);
+					NearestNeighborModel->SetNNINetwork(Network);
 					const bool bSuccess = NearestNeighborModel->LoadOptimizedNetwork(OnnxFile);
 					if (bSuccess)
 					{
@@ -115,7 +145,7 @@ namespace UE::NearestNeighborModel
 				}
 				else
 				{
-					Model->SetNeuralNetwork(Network);
+					NearestNeighborModel->SetNNINetwork(Network);
 					return true;
 				}
 			}
@@ -132,7 +162,7 @@ namespace UE::NearestNeighborModel
 		const UNearestNeighborModel* NearestNeighborModel = GetNearestNeighborModel();
 		if (NearestNeighborModel)
 		{
-			return NearestNeighborModel->DoesUseOptimizedNetwork() ? NearestNeighborModel->GetOptimizedNetwork() != nullptr : NearestNeighborModel->GetNeuralNetwork() != nullptr;
+			return NearestNeighborModel->DoesUseOptimizedNetwork() ? NearestNeighborModel->GetOptimizedNetwork() != nullptr : NearestNeighborModel->GetNNINetwork() != nullptr;
 		}
 		return false;
 	}
@@ -414,7 +444,7 @@ namespace UE::NearestNeighborModel
 			return EUpdateResult::SUCCESS;
 		}
 
-		const UNeuralNetwork *NeuralNetwork = NearestNeighborModel->GetNeuralNetwork();
+		const UNeuralNetwork* NeuralNetwork = NearestNeighborModel->GetNNINetwork();
 		if(NeuralNetwork && NeuralNetwork->IsLoaded())
 		{
 			const FString SavePath = GetTrainedNetworkOnnxFile();
@@ -505,10 +535,10 @@ namespace UE::NearestNeighborModel
 	void FNearestNeighborEditorModel::KMeansClusterPoses()
 	{
 		KMeansClusterResult = EUpdateResult::SUCCESS;
-		UNearestNeighborModel *NearestNeighborModel = static_cast<UNearestNeighborModel*>(Model);
+		UNearestNeighborModel* NearestNeighborModel = static_cast<UNearestNeighborModel*>(Model);
 		check(NearestNeighborModel != nullptr);
 
-		const UNeuralNetwork *NeuralNetwork = NearestNeighborModel->GetNeuralNetwork();
+		const UNeuralNetwork* NeuralNetwork = NearestNeighborModel->GetNNINetwork();
 		if(NeuralNetwork && NeuralNetwork->IsLoaded())
 		{
 			const FString SavePath = GetTrainedNetworkOnnxFile();
@@ -685,7 +715,7 @@ namespace UE::NearestNeighborModel
 
 		UNearestNeighborModel* NearestNeighborModel = GetNearestNeighborModel();
 		// Convert NNI network to  optimized network
-		const UNeuralNetwork *NeuralNetwork = NearestNeighborModel->GetNeuralNetwork();
+		const UNeuralNetwork* NeuralNetwork = NearestNeighborModel->GetNNINetwork();
 		if (NearestNeighborModel->DoesUseOptimizedNetwork() && NeuralNetwork != nullptr && NearestNeighborModel->GetOptimizedNetwork() == nullptr)
 		{
 			const FString SavePath = GetTrainedNetworkOnnxFile();
@@ -879,7 +909,7 @@ namespace UE::NearestNeighborModel
 		}
 		else
 		{
-			const UNeuralNetwork *NeuralNetwork = NearestNeighborModel->GetNeuralNetwork();
+			const UNeuralNetwork* NeuralNetwork = NearestNeighborModel->GetNNINetwork();
 			const UMLDeformerComponent* MLDeformerComponent = GetTestMLDeformerComponent();
 			if(NeuralNetwork && NeuralNetwork->IsLoaded() && MLDeformerComponent && MLDeformerComponent->GetModelInstance())
 			{
@@ -901,7 +931,7 @@ namespace UE::NearestNeighborModel
 	bool FNearestNeighborEditorModel::IsNeuralNetworkLoaded()
 	{
 		const UNearestNeighborModel *NearestNeighborModel = static_cast<UNearestNeighborModel*>(Model);
-		const UNeuralNetwork *NeuralNetwork = NearestNeighborModel->GetNeuralNetwork();
+		const UNeuralNetwork* NeuralNetwork = NearestNeighborModel->GetNNINetwork();
 		return NeuralNetwork && NeuralNetwork->IsLoaded();
 	}
 
diff --git a/Engine/Plugins/Animation/MLDeformer/NearestNeighborModel/Source/NearestNeighborModelEditor/Public/NearestNeighborEditorModel.h b/Engine/Plugins/Animation/MLDeformer/NearestNeighborModel/Source/NearestNeighborModelEditor/Public/NearestNeighborEditorModel.h
--- a/Engine/Plugins/Animation/MLDeformer/NearestNeighborModel/Source/NearestNeighborModelEditor/Public/NearestNeighborEditorModel.h
+++ b/Engine/Plugins/Animation/MLDeformer/NearestNeighborModel/Source/NearestNeighborModelEditor/Public/NearestNeighborEditorModel.h
@@ -74,6 +74,7 @@ namespace UE::NearestNeighborModel
 
 		int32 GetNumParts();
 
+		UNeuralNetwork* LoadNeuralNetworkFromOnnx(const FString& Filename) const;
 		void OnMorphTargetUpdate();
 		uint8 GetMorphTargetUpdateResult() { return MorphTargetUpdateResult; }
 		UMLDeformerComponent* GetTestMLDeformerComponent() const;
diff --git a/Engine/Plugins/Animation/MLDeformer/NeuralMorphModel/NeuralMorphModel.uplugin b/Engine/Plugins/Animation/MLDeformer/NeuralMorphModel/NeuralMorphModel.uplugin
--- a/Engine/Plugins/Animation/MLDeformer/NeuralMorphModel/NeuralMorphModel.uplugin
+++ b/Engine/Plugins/Animation/MLDeformer/NeuralMorphModel/NeuralMorphModel.uplugin
@@ -1,7 +1,7 @@
 {
 	"FileVersion": 3,
 	"Version": 1,
-	"VersionName": "0.1",
+	"VersionName": "0.2",
 	"FriendlyName": "ML Deformer Neural Morph Model",
 	"Description": "Neural Morph Model for the ML Deformer Framework",
 	"Category": "Animation",
@@ -31,10 +31,6 @@
 			"Name": "GeometryCache",
 			"Enabled": true
 		},
-		{
-			"Name": "NeuralNetworkInference",
-			"Enabled": true
-		},
 		{
 			"Name": "DeformerGraph",
 			"Enabled": true
diff --git a/Engine/Plugins/Animation/MLDeformer/NeuralMorphModel/Source/NeuralMorphModel/NeuralMorphModel.Build.cs b/Engine/Plugins/Animation/MLDeformer/NeuralMorphModel/Source/NeuralMorphModel/NeuralMorphModel.Build.cs
--- a/Engine/Plugins/Animation/MLDeformer/NeuralMorphModel/Source/NeuralMorphModel/NeuralMorphModel.Build.cs
+++ b/Engine/Plugins/Animation/MLDeformer/NeuralMorphModel/Source/NeuralMorphModel/NeuralMorphModel.Build.cs
@@ -38,7 +38,6 @@ public class NeuralMorphModel : ModuleRules
 				"CoreUObject",
 				"Engine",
 				"GeometryCache",
-				"NeuralNetworkInference",
 				"OptimusCore",
 				"Projects",
 				"RenderCore",
diff --git a/Engine/Plugins/Animation/MLDeformer/NeuralMorphModel/Source/NeuralMorphModel/Private/NeuralMorphModel.cpp b/Engine/Plugins/Animation/MLDeformer/NeuralMorphModel/Source/NeuralMorphModel/Private/NeuralMorphModel.cpp
--- a/Engine/Plugins/Animation/MLDeformer/NeuralMorphModel/Source/NeuralMorphModel/Private/NeuralMorphModel.cpp
+++ b/Engine/Plugins/Animation/MLDeformer/NeuralMorphModel/Source/NeuralMorphModel/Private/NeuralMorphModel.cpp
@@ -17,12 +17,6 @@ namespace UE::NeuralMorphModel
 	class NEURALMORPHMODEL_API FNeuralMorphModelModule
 		: public IModuleInterface
 	{
-		void StartupModule() override
-		{
-			#if NEURALMORPHMODEL_FORCE_USE_NNI
-				UE_LOG(LogNeuralMorphModel, Warning, TEXT("Running neural morph model with NNI. The faster custom inference code path will be disabled."));
-			#endif
-		}
 	};
 }
 IMPLEMENT_MODULE(UE::NeuralMorphModel::FNeuralMorphModelModule, NeuralMorphModel)
@@ -55,26 +49,15 @@ UMLDeformerInputInfo* UNeuralMorphModel::CreateInputInfo()
 
 void UNeuralMorphModel::Serialize(FArchive& Archive)
 {
-#if !NEURALMORPHMODEL_FORCE_USE_NNI
-	if (Archive.IsSaving() || Archive.IsCooking())
+	if (Archive.IsSaving() && Archive.IsCooking())
 	{
-		UNeuralNetwork* NNINeuralNetwork = GetNeuralNetwork();
-
-		// Show a warning when we are not using custom inference yet on this model.
-		if (NeuralMorphNetwork == nullptr && NNINeuralNetwork != nullptr)
+		if (NeuralMorphNetwork == nullptr)
 		{
-			UE_LOG(LogNeuralMorphModel, Display, TEXT("Neural Morph Model in MLD asset '%s' should be retrained to get higher performance by taking advantage of custom inference."), *GetDeformerAsset()->GetName());
-		}
-
-		// If we have a custom inference network, make sure we don't save out the NNI network.
-		if (NeuralMorphNetwork != nullptr && NNINeuralNetwork != nullptr)
-		{
-			SetNeuralNetwork(nullptr);
+			UE_LOG(LogNeuralMorphModel, Display, TEXT("Neural Morph Model in MLD asset '%s' still needs to be trained."), *GetDeformerAsset()->GetName());
 		}
 	}
-#endif
 
-	// Convert the UMLDeformerInputInfo object into a UNeuralMorphInputInfo object.
+	// Convert the UMLDeformerInputInfo object into a UNeuralMorphInputInfo object for backward compatiblity.
 	UMLDeformerInputInfo* CurInputInfo = GetInputInfo();
 	if (CurInputInfo)
 	{
@@ -89,17 +72,4 @@ void UNeuralMorphModel::Serialize(FArchive& Archive)
 	Super::Serialize(Archive);
 }
 
-void UNeuralMorphModel::PostLoad()
-{
-	Super::PostLoad();
-
-#if !NEURALMORPHMODEL_FORCE_USE_NNI
-	// Show a warning when we are not using custom inference yet on this model.
-	if (NeuralMorphNetwork == nullptr && GetNeuralNetwork() != nullptr)
-	{
-		UE_LOG(LogNeuralMorphModel, Display, TEXT("Neural Morph Model in MLD asset '%s' should be retrained to get higher performance by taking advantage of custom inference."), *GetDeformerAsset()->GetName());
-	}
-#endif
-}
-
 #undef LOCTEXT_NAMESPACE
diff --git a/Engine/Plugins/Animation/MLDeformer/NeuralMorphModel/Source/NeuralMorphModel/Private/NeuralMorphModelInstance.cpp b/Engine/Plugins/Animation/MLDeformer/NeuralMorphModel/Source/NeuralMorphModel/Private/NeuralMorphModelInstance.cpp
--- a/Engine/Plugins/Animation/MLDeformer/NeuralMorphModel/Source/NeuralMorphModel/Private/NeuralMorphModelInstance.cpp
+++ b/Engine/Plugins/Animation/MLDeformer/NeuralMorphModel/Source/NeuralMorphModel/Private/NeuralMorphModelInstance.cpp
@@ -15,11 +15,8 @@ void UNeuralMorphModelInstance::Init(USkeletalMeshComponent* SkelMeshComponent)
 	Super::Init(SkelMeshComponent);
 
 	UNeuralMorphModel* MorphModel = Cast<UNeuralMorphModel>(Model);
-
-#if !NEURALMORPHMODEL_FORCE_USE_NNI
 	UNeuralMorphNetwork* MorphNetwork = MorphModel->GetNeuralMorphNetwork();
 	NetworkInstance = MorphNetwork ? MorphNetwork->CreateInstance() : nullptr;
-#endif
 
 	UNeuralMorphInputInfo* InputInfo = Cast<UNeuralMorphInputInfo>(Model->GetInputInfo());
 	if (InputInfo == nullptr)
@@ -186,21 +183,17 @@ bool UNeuralMorphModelInstance::SetupInputs()
 {
 	TRACE_CPUPROFILER_EVENT_SCOPE(UNeuralMorphModelInstance::SetupInputs)
 
-#if NEURALMORPHMODEL_FORCE_USE_NNI
-	return Super::SetupInputs();
-#else
-	// If we have no neural morph network, fall back to the default NNI path.
+	// Make sure we have a valid network.
 	UNeuralMorphModel* MorphModel = Cast<UNeuralMorphModel>(Model);
 	UNeuralMorphNetwork* MorphNetwork = MorphModel->GetNeuralMorphNetwork();
-	if (MorphNetwork == nullptr)
+	if (MorphNetwork == nullptr || MorphNetwork->GetMainMLP() == nullptr)
 	{		
-		return Super::SetupInputs();
+		return false;
 	}
 
-	// Some safety checks.
+	// Some additional safety checks.
 	if (SkeletalMeshComponent == nullptr ||
 		SkeletalMeshComponent->GetSkeletalMeshAsset() == nullptr ||
-		MorphNetwork->GetMainMLP() == nullptr ||
 		!bIsCompatible)
 	{
 		return false;
@@ -220,17 +213,12 @@ bool UNeuralMorphModelInstance::SetupInputs()
 	FillNetworkInputs();
 
 	return true;
-#endif
 }
 
 void UNeuralMorphModelInstance::Execute(float ModelWeight)
 {
 	TRACE_CPUPROFILER_EVENT_SCOPE(UNeuralMorphModelInstance::Execute)
 
-#if NEURALMORPHMODEL_FORCE_USE_NNI
-	Super::Execute(ModelWeight);
-	return;
-#else
 	UNeuralMorphModel* MorphModel = Cast<UNeuralMorphModel>(Model);
 	UNeuralMorphNetwork* MorphNetwork = MorphModel->GetNeuralMorphNetwork();
 	if (MorphNetwork == nullptr)
@@ -303,5 +291,4 @@ void UNeuralMorphModelInstance::Execute(float ModelWeight)
 			WeightData->Weights[MorphIndex + 1] = NetworkOutputs[MorphIndex] * ModelWeight;
 		}
 	}
-#endif
 }
diff --git a/Engine/Plugins/Animation/MLDeformer/NeuralMorphModel/Source/NeuralMorphModel/Public/NeuralMorphModel.h b/Engine/Plugins/Animation/MLDeformer/NeuralMorphModel/Source/NeuralMorphModel/Public/NeuralMorphModel.h
--- a/Engine/Plugins/Animation/MLDeformer/NeuralMorphModel/Source/NeuralMorphModel/Public/NeuralMorphModel.h
+++ b/Engine/Plugins/Animation/MLDeformer/NeuralMorphModel/Source/NeuralMorphModel/Public/NeuralMorphModel.h
@@ -11,15 +11,11 @@ class USkeletalMesh;
 class UGeometryCache;
 class UAnimSequence;
 class UNeuralMorphNetwork;
-class UNeuralNetwork;
 class UMLDeformerAsset;
 class UMLDeformerModelInstance;
 class USkeleton;
 struct FExternalMorphSet;
 
-/** Set this to 1 if you like to force to use NNI for inference, rather than the (faster) custom inference, otherwise set to 0. */
-#define NEURALMORPHMODEL_FORCE_USE_NNI 0
-
 /** 
  * Specify whether we want to use ISPC if available.
  * You can comment out this line, to force disable the use of ISPC in this plugin.
@@ -53,7 +49,6 @@ public:
 
 	// UObject overrides.
 	virtual void Serialize(FArchive& Archive) override;
-	virtual void PostLoad() override;
 	// ~END UObject overrides.
 
 	// UMLDeformerModel overrides.
diff --git a/Engine/Plugins/Animation/MLDeformer/NeuralMorphModel/Source/NeuralMorphModelEditor/NeuralMorphModelEditor.Build.cs b/Engine/Plugins/Animation/MLDeformer/NeuralMorphModel/Source/NeuralMorphModelEditor/NeuralMorphModelEditor.Build.cs
--- a/Engine/Plugins/Animation/MLDeformer/NeuralMorphModel/Source/NeuralMorphModelEditor/NeuralMorphModelEditor.Build.cs
+++ b/Engine/Plugins/Animation/MLDeformer/NeuralMorphModel/Source/NeuralMorphModelEditor/NeuralMorphModelEditor.Build.cs
@@ -31,7 +31,6 @@ namespace UnrealBuildTool.Rules
 					"MLDeformerFrameworkEditor",
 					"NeuralMorphModel",
 					"PropertyEditor",
-					"NeuralNetworkInference",
 					"ToolWidgets",
 					"ComputeFramework",
 					"RenderCore",
diff --git a/Engine/Plugins/Animation/MLDeformer/NeuralMorphModel/Source/NeuralMorphModelEditor/Private/NeuralMorphEditorModel.cpp b/Engine/Plugins/Animation/MLDeformer/NeuralMorphModel/Source/NeuralMorphModelEditor/Private/NeuralMorphEditorModel.cpp
--- a/Engine/Plugins/Animation/MLDeformer/NeuralMorphModel/Source/NeuralMorphModelEditor/Private/NeuralMorphEditorModel.cpp
+++ b/Engine/Plugins/Animation/MLDeformer/NeuralMorphModel/Source/NeuralMorphModelEditor/Private/NeuralMorphEditorModel.cpp
@@ -55,14 +55,6 @@ namespace UE::NeuralMorphModel
 
 	bool FNeuralMorphEditorModel::IsTrained() const
 	{
-#if NEURALMORPHMODEL_FORCE_USE_NNI
-		return GetMorphModel()->GetNeuralNetwork() != nullptr;
-#else
-		if (GetMorphModel()->GetNeuralNetwork() != nullptr)
-		{
-			return true;
-		}
-
 		UNeuralMorphNetwork* MorphNetwork = GetNeuralMorphModel()->GetNeuralMorphNetwork();
 		if (MorphNetwork != nullptr && MorphNetwork->GetMainMLP() != nullptr)
 		{
@@ -70,7 +62,6 @@ namespace UE::NeuralMorphModel
 		}
 
 		return false;
-#endif
 	}
 
 	ETrainingResult FNeuralMorphEditorModel::Train()
@@ -80,14 +71,6 @@ namespace UE::NeuralMorphModel
 
 	bool FNeuralMorphEditorModel::LoadTrainedNetwork() const
 	{
-#if NEURALMORPHMODEL_FORCE_USE_NNI
-		const bool bSuccess = FMLDeformerMorphModelEditorModel::LoadTrainedNetwork();
-		if (bSuccess)
-		{
-			GetNeuralMorphModel()->SetNeuralMorphNetwork(nullptr);	// Force disable custom inference.
-		}
-		return bSuccess;
-#else
 		// Load the specialized neural morph model network.
 		// Base the filename on the onnx filename, and replace the file extension.
 		FString NetworkFilename = GetTrainedNetworkOnnxFile();
@@ -112,9 +95,7 @@ namespace UE::NeuralMorphModel
 			NeuralNet = nullptr;
 		}
 
-		GetNeuralMorphModel()->SetNeuralNetwork(nullptr);			// Disable NNI inference.
 		GetNeuralMorphModel()->SetNeuralMorphNetwork(NeuralNet);	// Use our custom inference.
-#endif
 		return true;
 	}
 
diff --git a/Engine/Plugins/Animation/MLDeformer/NeuralMorphModel/Source/NeuralMorphModelEditor/Private/NeuralMorphModelDetails.cpp b/Engine/Plugins/Animation/MLDeformer/NeuralMorphModel/Source/NeuralMorphModelEditor/Private/NeuralMorphModelDetails.cpp
--- a/Engine/Plugins/Animation/MLDeformer/NeuralMorphModel/Source/NeuralMorphModelEditor/Private/NeuralMorphModelDetails.cpp
+++ b/Engine/Plugins/Animation/MLDeformer/NeuralMorphModel/Source/NeuralMorphModelEditor/Private/NeuralMorphModelDetails.cpp
@@ -6,7 +6,6 @@
 #include "NeuralMorphEditorModel.h"
 #include "DetailCategoryBuilder.h"
 #include "DetailWidgetRow.h"
-
 #include "Widgets/Layout/SBox.h"
 #include "Widgets/Input/SButton.h"
 #include "SWarningOrErrorBox.h"
@@ -69,44 +68,6 @@ namespace UE::NeuralMorphModel
 		InputOutputCategoryBuilder->AddProperty(GET_MEMBER_NAME_CHECKED(UNeuralMorphModel, CurveGroups), UNeuralMorphModel::StaticClass())
 			.Visibility(NeuralMorphModel->Mode == ENeuralMorphMode::Local ? EVisibility::Visible : EVisibility::Collapsed);
 	}
-
-	void FNeuralMorphModelDetails::AddTrainingSettingsErrors()
-	{
-		FMLDeformerMorphModelDetails::AddTrainingSettingsErrors();
-
-		UNeuralMorphModel* NeuralMorphModel = Cast<UNeuralMorphModel>(Model);
-		check(NeuralMorphModel);
-
-		#if !NEURALMORPHMODEL_FORCE_USE_NNI
-			const bool bShowWarning = (NeuralMorphModel->GetNeuralNetwork() != nullptr) && (NeuralMorphModel->GetNeuralMorphNetwork() == nullptr);
-			FDetailWidgetRow& PerformanceWarningRow = TrainingSettingsCategoryBuilder->AddCustomRow(FText::FromString("NeuralNetPerformanceWarning"))
-				.Visibility(bShowWarning ? EVisibility::Visible : EVisibility::Collapsed)
-				.WholeRowContent()
-				[
-					SNew(SBox)
-					.Padding(FMargin(0.0f, 4.0f))
-					[
-						SNew(SWarningOrErrorBox)
-						.MessageStyle(EMessageStyle::Warning)
-						.Message(LOCTEXT("NeuralNetPerformanceWarning", "The model must be retrained in order to make use of higher performance inference."))
-					]
-				];
-		#else
-			const bool bShowWarning = (NeuralMorphModel->GetNeuralNetwork() == nullptr) && (NeuralMorphModel->GetNeuralMorphNetwork() != nullptr);
-			FDetailWidgetRow& PerformanceWarningRow = TrainingSettingsCategoryBuilder->AddCustomRow(FText::FromString("NeuralNetWrongInferenceWarning"))
-				.Visibility(bShowWarning ? EVisibility::Visible : EVisibility::Collapsed)
-				.WholeRowContent()
-				[
-					SNew(SBox)
-					.Padding(FMargin(0.0f, 4.0f))
-					[
-						SNew(SWarningOrErrorBox)
-						.MessageStyle(EMessageStyle::Warning)
-						.Message(LOCTEXT("NeuralNetWrongInferenceWarning", "The model was trained using custom inference, but NNI is set as active inference engine. Please retrain the model or switch back to custom inference by recompiling the plugin."))
-					]
-				];
-		#endif
-	}
 }	// namespace UE::NeuralMorphModel
 
 #undef LOCTEXT_NAMESPACE
diff --git a/Engine/Plugins/Animation/MLDeformer/NeuralMorphModel/Source/NeuralMorphModelEditor/Private/NeuralMorphTrainingModel.cpp b/Engine/Plugins/Animation/MLDeformer/NeuralMorphModel/Source/NeuralMorphModelEditor/Private/NeuralMorphTrainingModel.cpp
--- a/Engine/Plugins/Animation/MLDeformer/NeuralMorphModel/Source/NeuralMorphModelEditor/Private/NeuralMorphTrainingModel.cpp
+++ b/Engine/Plugins/Animation/MLDeformer/NeuralMorphModel/Source/NeuralMorphModelEditor/Private/NeuralMorphTrainingModel.cpp
@@ -5,16 +5,6 @@
 #include "NeuralMorphModel.h"
 #include "NeuralMorphInputInfo.h"
 
-
-bool UNeuralMorphTrainingModel::ForceUseNNI() const
-{
-	#if NEURALMORPHMODEL_FORCE_USE_NNI
-		return true;
-	#else
-		return false;
-	#endif
-}
-
 int32 UNeuralMorphTrainingModel::GetNumBoneGroups() const
 {
 	UNeuralMorphInputInfo* InputInfo = Cast<UNeuralMorphInputInfo>(EditorModel->GetEditorInputInfo());
diff --git a/Engine/Plugins/Animation/MLDeformer/NeuralMorphModel/Source/NeuralMorphModelEditor/Public/NeuralMorphModelDetails.h b/Engine/Plugins/Animation/MLDeformer/NeuralMorphModel/Source/NeuralMorphModelEditor/Public/NeuralMorphModelDetails.h
--- a/Engine/Plugins/Animation/MLDeformer/NeuralMorphModel/Source/NeuralMorphModelEditor/Public/NeuralMorphModelDetails.h
+++ b/Engine/Plugins/Animation/MLDeformer/NeuralMorphModel/Source/NeuralMorphModelEditor/Public/NeuralMorphModelDetails.h
@@ -25,7 +25,6 @@ namespace UE::NeuralMorphModel
 		// ~END ILayoutDetails overrides.
 
 		// FMLDeformerModelDetails overrides.
-		virtual void AddTrainingSettingsErrors() override;
 		virtual void AddTrainingInputFilters() override;
 		// ~END FMLDeformerModelDetails overrides.
 	};
diff --git a/Engine/Plugins/Animation/MLDeformer/NeuralMorphModel/Source/NeuralMorphModelEditor/Public/NeuralMorphTrainingModel.h b/Engine/Plugins/Animation/MLDeformer/NeuralMorphModel/Source/NeuralMorphModelEditor/Public/NeuralMorphTrainingModel.h
--- a/Engine/Plugins/Animation/MLDeformer/NeuralMorphModel/Source/NeuralMorphModelEditor/Public/NeuralMorphTrainingModel.h
+++ b/Engine/Plugins/Animation/MLDeformer/NeuralMorphModel/Source/NeuralMorphModelEditor/Public/NeuralMorphTrainingModel.h
@@ -26,13 +26,6 @@ public:
 	UFUNCTION(BlueprintImplementableEvent, Category = "Training Model")
 	int32 Train() const;
 
-	/**
-	 * Do we want to force using NNI for inference, instead of our custom inference?
-	 * @result Returns true when NNI should be used for inference, or false if custom (faster) inference should be used.
-	 */
-	UFUNCTION(BlueprintPure, Category = "Training Model")
-	bool ForceUseNNI() const;
-
 	UFUNCTION(BlueprintPure, Category = "Training Model")
 	int32 GetNumBoneGroups() const;
 
diff --git a/Engine/Plugins/Animation/MLDeformer/VertexDeltaModel/Shaders/Private/VertexDeltaModel.ush b/Engine/Plugins/Animation/MLDeformer/VertexDeltaModel/Shaders/Private/VertexDeltaModel.ush
new file mode 100644
--- /dev/null
+++ b/Engine/Plugins/Animation/MLDeformer/VertexDeltaModel/Shaders/Private/VertexDeltaModel.ush
@@ -0,0 +1,19 @@
+// Copyright Epic Games, Inc. All Rights Reserved.
+
+uint {DataInterfaceName}_NumVertices;
+uint {DataInterfaceName}_InputStreamStart;
+float {DataInterfaceName}_Weight;
+Buffer<float> {DataInterfaceName}_PositionDeltaBuffer;
+Buffer<uint> {DataInterfaceName}_VertexMapBuffer;
+
+uint ReadNumVertices_{DataInterfaceName}()
+{
+	return {DataInterfaceName}_NumVertices;
+}
+
+float3 ReadPositionDelta_{DataInterfaceName}(uint VertexIndex)
+{
+	uint BufferIndex = {DataInterfaceName}_VertexMapBuffer[VertexIndex + {DataInterfaceName}_InputStreamStart] * 3;
+	float3 Delta = float3({DataInterfaceName}_PositionDeltaBuffer[BufferIndex], {DataInterfaceName}_PositionDeltaBuffer[BufferIndex + 1], {DataInterfaceName}_PositionDeltaBuffer[BufferIndex + 2]);
+	return Delta * {DataInterfaceName}_Weight;
+}
diff --git a/Engine/Plugins/Animation/MLDeformer/VertexDeltaModel/Shaders/Private/VertexDeltaModelHeatMap.ush b/Engine/Plugins/Animation/MLDeformer/VertexDeltaModel/Shaders/Private/VertexDeltaModelHeatMap.ush
new file mode 100644
--- /dev/null
+++ b/Engine/Plugins/Animation/MLDeformer/VertexDeltaModel/Shaders/Private/VertexDeltaModelHeatMap.ush
@@ -0,0 +1,41 @@
+// Copyright Epic Games, Inc. All Rights Reserved.
+
+uint {DataInterfaceName}_NumVertices;
+uint {DataInterfaceName}_InputStreamStart;
+uint {DataInterfaceName}_HeatMapMode;
+float {DataInterfaceName}_HeatMapMax;
+float {DataInterfaceName}_GroundTruthLerp;
+uint {DataInterfaceName}_GroundTruthBufferSize;
+StructuredBuffer<float> {DataInterfaceName}_PositionGroundTruthBuffer;
+Buffer<uint> {DataInterfaceName}_VertexMapBuffer;
+
+uint ReadNumVertices_{DataInterfaceName}()
+{
+	return {DataInterfaceName}_NumVertices;
+}
+
+float3 ReadPositionGroundTruth_{DataInterfaceName}(uint VertexIndex)
+{
+	uint BufferIndex = {DataInterfaceName}_VertexMapBuffer[VertexIndex + {DataInterfaceName}_InputStreamStart];
+	if (BufferIndex >= {DataInterfaceName}_GroundTruthBufferSize)
+	{
+		return float3(0, 0, 0);
+	}
+	BufferIndex *= 3;
+	return float3({DataInterfaceName}_PositionGroundTruthBuffer[BufferIndex], {DataInterfaceName}_PositionGroundTruthBuffer[BufferIndex + 1], {DataInterfaceName}_PositionGroundTruthBuffer[BufferIndex + 2]);
+}
+
+uint ReadHeatMapMode_{DataInterfaceName}()
+{
+	return {DataInterfaceName}_HeatMapMode;
+}
+
+float ReadHeatMapMax_{DataInterfaceName}()
+{
+	return {DataInterfaceName}_HeatMapMax;
+}
+
+float ReadGroundTruthLerp_{DataInterfaceName}()
+{
+	return {DataInterfaceName}_GroundTruthLerp;
+}
diff --git a/Engine/Plugins/Animation/MLDeformer/VertexDeltaModel/Source/VertexDeltaModel/Private/VertexDeltaGraphDataInterface.cpp b/Engine/Plugins/Animation/MLDeformer/VertexDeltaModel/Source/VertexDeltaModel/Private/VertexDeltaGraphDataInterface.cpp
new file mode 100644
--- /dev/null
+++ b/Engine/Plugins/Animation/MLDeformer/VertexDeltaModel/Source/VertexDeltaModel/Private/VertexDeltaGraphDataInterface.cpp
@@ -0,0 +1,158 @@
+// Copyright Epic Games, Inc. All Rights Reserved.
+
+#include "VertexDeltaGraphDataInterface.h"
+#include "VertexDeltaModel.h"
+#include "Components/SkeletalMeshComponent.h"
+#include "ComputeFramework/ShaderParamTypeDefinition.h"
+#include "MLDeformerModel.h"
+#include "MLDeformerModelInstance.h"
+#include "MLDeformerAsset.h"
+#include "MLDeformerComponent.h"
+#include "NeuralNetwork.h"
+#include "OptimusDataDomain.h"
+#include "RenderGraphBuilder.h"
+#include "RenderGraphUtils.h"
+#include "RenderGraphResources.h"
+#include "ShaderCompilerCore.h"
+#include "ShaderParameterMetadataBuilder.h"
+#include "SkeletalRenderPublic.h"
+
+#include UE_INLINE_GENERATED_CPP_BY_NAME(VertexDeltaGraphDataInterface)
+
+TArray<FOptimusCDIPinDefinition> UVertexDeltaGraphDataInterface::GetPinDefinitions() const
+{
+	TArray<FOptimusCDIPinDefinition> Defs;
+	Defs.Add({ "PositionDelta", "ReadPositionDelta", Optimus::DomainName::Vertex, "ReadNumVertices" });
+	return Defs;
+}
+
+TSubclassOf<UActorComponent> UVertexDeltaGraphDataInterface::GetRequiredComponentClass() const
+{
+	return UMLDeformerComponent::StaticClass();
+}
+
+void UVertexDeltaGraphDataInterface::GetSupportedInputs(TArray<FShaderFunctionDefinition>& OutFunctions) const
+{
+	OutFunctions.AddDefaulted_GetRef()
+		.SetName(TEXT("ReadNumVertices"))
+		.AddReturnType(EShaderFundamentalType::Uint);
+
+	OutFunctions.AddDefaulted_GetRef()
+		.SetName(TEXT("ReadPositionDelta"))
+		.AddReturnType(EShaderFundamentalType::Float, 3)
+		.AddParam(EShaderFundamentalType::Uint);
+}
+
+BEGIN_SHADER_PARAMETER_STRUCT(FVertexDeltaGraphDataInterfaceParameters, )
+	SHADER_PARAMETER(uint32, NumVertices)
+	SHADER_PARAMETER(uint32, InputStreamStart)
+	SHADER_PARAMETER(float, Weight)
+	SHADER_PARAMETER_RDG_BUFFER_SRV(Buffer<float>, PositionDeltaBuffer)
+	SHADER_PARAMETER_SRV(Buffer<uint>, VertexMapBuffer)
+END_SHADER_PARAMETER_STRUCT()
+
+FString UVertexDeltaGraphDataInterface::GetDisplayName() const
+{
+	return TEXT("MLD Vertex Delta Model");
+}
+
+void UVertexDeltaGraphDataInterface::GetShaderParameters(TCHAR const* UID, FShaderParametersMetadataBuilder& InOutBuilder, FShaderParametersMetadataAllocations& InOutAllocations) const \
+{
+	InOutBuilder.AddNestedStruct<FVertexDeltaGraphDataInterfaceParameters>(UID);
+}
+
+TCHAR const* UVertexDeltaGraphDataInterface::TemplateFilePath = TEXT("/Plugin/VertexDeltaModel/Private/VertexDeltaModel.ush");
+
+TCHAR const* UVertexDeltaGraphDataInterface::GetShaderVirtualPath() const
+{
+	return TemplateFilePath;
+}
+
+void UVertexDeltaGraphDataInterface::GetShaderHash(FString& InOutKey) const
+{
+	GetShaderFileHash(TemplateFilePath, EShaderPlatform::SP_PCD3D_SM5).AppendString(InOutKey);
+}
+
+void UVertexDeltaGraphDataInterface::GetHLSL(FString& OutHLSL, FString const& InDataInterfaceName) const
+{
+	TMap<FString, FStringFormatArg> TemplateArgs =
+	{
+		{ TEXT("DataInterfaceName"), InDataInterfaceName },
+	};
+
+	FString TemplateFile;
+	LoadShaderSourceFile(TemplateFilePath, EShaderPlatform::SP_PCD3D_SM5, &TemplateFile, nullptr);
+	OutHLSL += FString::Format(*TemplateFile, TemplateArgs);
+}
+
+UComputeDataProvider* UVertexDeltaGraphDataInterface::CreateDataProvider(TObjectPtr<UObject> InBinding, uint64 InInputMask, uint64 InOutputMask) const
+{
+	UVertexDeltaGraphDataProvider* Provider = NewObject<UVertexDeltaGraphDataProvider>();
+	Provider->DeformerComponent = Cast<UMLDeformerComponent>(InBinding);
+	return Provider;
+}
+
+FComputeDataProviderRenderProxy* UVertexDeltaGraphDataProvider::GetRenderProxy()
+{
+	return new UE::VertexDeltaModel::FVertexDeltaGraphDataProviderProxy(DeformerComponent);
+}
+
+namespace UE::VertexDeltaModel
+{
+	FVertexDeltaGraphDataProviderProxy::FVertexDeltaGraphDataProviderProxy(UMLDeformerComponent* DeformerComponent)
+	{
+		using namespace UE::MLDeformer;
+
+		const UMLDeformerAsset* DeformerAsset = DeformerComponent != nullptr ? DeformerComponent->GetDeformerAsset() : nullptr;
+		const UMLDeformerModel* Model = DeformerAsset != nullptr ? DeformerAsset->GetModel() : nullptr;
+		const UMLDeformerModelInstance* ModelInstance = DeformerComponent != nullptr ? DeformerComponent->GetModelInstance() : nullptr;
+		if (Model != nullptr && ModelInstance != nullptr)
+		{
+			const UVertexDeltaModel* VertexDeltaModel = Cast<UVertexDeltaModel>(Model);
+			SkeletalMeshObject = ModelInstance->GetSkeletalMeshComponent()->MeshObject;
+			NeuralNetwork = VertexDeltaModel->GetNNINetwork();
+			NeuralNetworkInferenceHandle = ModelInstance->GetNeuralNetworkInferenceHandle();
+			bCanRunNeuralNet = ModelInstance->IsCompatible() && Model->IsNeuralNetworkOnGPU() && DeformerComponent->GetModelInstance()->IsValidForDataProvider();
+			Weight = DeformerComponent->GetWeight();
+			VertexMapBufferSRV = VertexDeltaModel->GetVertexMapBuffer().ShaderResourceViewRHI;
+		}
+	}
+
+	bool FVertexDeltaGraphDataProviderProxy::IsValid(FValidationData const& InValidationData) const
+	{
+		if (InValidationData.ParameterStructSize != sizeof(FVertexDeltaGraphDataInterfaceParameters))
+		{
+			return false;
+		}
+
+		if (!bCanRunNeuralNet || SkeletalMeshObject == nullptr || NeuralNetwork == nullptr || VertexMapBufferSRV == nullptr)
+		{
+			return false;
+		}
+
+		return true;
+	}
+
+	void FVertexDeltaGraphDataProviderProxy::AllocateResources(FRDGBuilder& GraphBuilder)
+	{
+		Buffer = GraphBuilder.RegisterExternalBuffer(NeuralNetwork->GetOutputTensorForContext(NeuralNetworkInferenceHandle).GetPooledBuffer());
+		BufferSRV = GraphBuilder.CreateSRV(Buffer, PF_R32_FLOAT);
+	}
+
+	void FVertexDeltaGraphDataProviderProxy::GatherDispatchData(FDispatchData const& InDispatchData)
+	{
+		const FSkeletalMeshRenderData& SkeletalMeshRenderData = SkeletalMeshObject->GetSkeletalMeshRenderData();
+		const FSkeletalMeshLODRenderData* LodRenderData = SkeletalMeshRenderData.GetPendingFirstLOD(0);
+		const TStridedView<FVertexDeltaGraphDataInterfaceParameters> ParameterArray = MakeStridedParameterView<FVertexDeltaGraphDataInterfaceParameters>(InDispatchData);
+		for (int32 InvocationIndex = 0; InvocationIndex < ParameterArray.Num(); ++InvocationIndex)
+		{
+			const FSkelMeshRenderSection& RenderSection = LodRenderData->RenderSections[InvocationIndex];
+			FVertexDeltaGraphDataInterfaceParameters& Parameters = ParameterArray[InvocationIndex];
+			Parameters.NumVertices = 0;
+			Parameters.InputStreamStart = RenderSection.BaseVertexIndex;
+			Parameters.Weight = Weight;
+			Parameters.PositionDeltaBuffer = BufferSRV;
+			Parameters.VertexMapBuffer = VertexMapBufferSRV;
+		}
+	}
+}	// namespace UE::VertexDeltaModel
diff --git a/Engine/Plugins/Animation/MLDeformer/VertexDeltaModel/Source/VertexDeltaModel/Private/VertexDeltaGraphDebugDataInterface.cpp b/Engine/Plugins/Animation/MLDeformer/VertexDeltaModel/Source/VertexDeltaModel/Private/VertexDeltaGraphDebugDataInterface.cpp
new file mode 100644
--- /dev/null
+++ b/Engine/Plugins/Animation/MLDeformer/VertexDeltaModel/Source/VertexDeltaModel/Private/VertexDeltaGraphDebugDataInterface.cpp
@@ -0,0 +1,225 @@
+// Copyright Epic Games, Inc. All Rights Reserved.
+
+#include "VertexDeltaGraphDebugDataInterface.h"
+#include "VertexDeltaGraphDataInterface.h"
+#include "MLDeformerModelInstance.h"
+#include "MLDeformerAsset.h"
+#include "MLDeformerComponent.h"
+#include "MLDeformerVizSettings.h"
+#include "Components/SkeletalMeshComponent.h"
+#include "ComputeFramework/ComputeKernelPermutationSet.h"
+#include "ComputeFramework/ShaderParamTypeDefinition.h"
+#include "OptimusDataDomain.h"
+#include "RenderGraphBuilder.h"
+#include "RenderGraphUtils.h"
+#include "RenderGraphResources.h"
+#include "Rendering/SkeletalMeshModel.h"
+#include "ShaderCompilerCore.h"
+#include "ShaderParameterMetadataBuilder.h"
+#include "SkeletalRenderPublic.h"
+#include "NeuralNetwork.h"
+
+#include UE_INLINE_GENERATED_CPP_BY_NAME(VertexDeltaGraphDebugDataInterface)
+
+TArray<FOptimusCDIPinDefinition> UVertexDeltaGraphDebugDataInterface::GetPinDefinitions() const
+{
+	TArray<FOptimusCDIPinDefinition> Defs;
+	Defs.Add({ "HeatMapMode", "ReadHeatMapMode" });
+	Defs.Add({ "HeatMapMax", "ReadHeatMapMax" });
+	Defs.Add({ "GroundTruthLerp", "ReadGroundTruthLerp" });
+	Defs.Add({ "PositionGroundTruth", "ReadPositionGroundTruth", Optimus::DomainName::Vertex, "ReadNumVertices" });
+	return Defs;
+}
+
+TSubclassOf<UActorComponent> UVertexDeltaGraphDebugDataInterface::GetRequiredComponentClass() const
+{
+	return UMLDeformerComponent::StaticClass();
+}
+
+void UVertexDeltaGraphDebugDataInterface::GetSupportedInputs(TArray<FShaderFunctionDefinition>& OutFunctions) const
+{
+	OutFunctions.AddDefaulted_GetRef()
+		.SetName(TEXT("ReadNumVertices"))
+		.AddReturnType(EShaderFundamentalType::Uint);
+
+	OutFunctions.AddDefaulted_GetRef()
+		.SetName(TEXT("ReadHeatMapMode"))
+		.AddReturnType(EShaderFundamentalType::Int);
+
+	OutFunctions.AddDefaulted_GetRef()
+		.SetName(TEXT("ReadHeatMapMax"))
+		.AddReturnType(EShaderFundamentalType::Float);
+
+	OutFunctions.AddDefaulted_GetRef()
+		.SetName(TEXT("ReadGroundTruthLerp"))
+		.AddReturnType(EShaderFundamentalType::Float);
+
+	OutFunctions.AddDefaulted_GetRef()
+		.SetName(TEXT("ReadPositionGroundTruth"))
+		.AddReturnType(EShaderFundamentalType::Float, 3)
+		.AddParam(EShaderFundamentalType::Uint);
+}
+
+BEGIN_SHADER_PARAMETER_STRUCT(FVertexDeltaGraphDebugDataInterfaceParameters, )
+	SHADER_PARAMETER(uint32, NumVertices)
+	SHADER_PARAMETER(uint32, InputStreamStart)
+	SHADER_PARAMETER(int32, HeatMapMode)
+	SHADER_PARAMETER(float, HeatMapMax)
+	SHADER_PARAMETER(float, GroundTruthLerp)
+	SHADER_PARAMETER(uint32, GroundTruthBufferSize)
+	SHADER_PARAMETER_RDG_BUFFER_SRV(StructuredBuffer<float>, PositionGroundTruthBuffer)
+	SHADER_PARAMETER_SRV(Buffer<uint>, VertexMapBuffer)
+END_SHADER_PARAMETER_STRUCT()
+
+FString UVertexDeltaGraphDebugDataInterface::GetDisplayName() const
+{
+	return TEXT("MLD Vertex Delta Model Debug");
+}
+
+void UVertexDeltaGraphDebugDataInterface::GetShaderParameters(TCHAR const* UID, FShaderParametersMetadataBuilder& InOutBuilder, FShaderParametersMetadataAllocations& InOutAllocations) const
+{
+	InOutBuilder.AddNestedStruct<FVertexDeltaGraphDebugDataInterfaceParameters>(UID);
+}
+
+TCHAR const* UVertexDeltaGraphDebugDataInterface::TemplateFilePath = TEXT("/Plugin/VertexDeltaModel/Private/VertexDeltaModelHeatMap.ush");
+
+TCHAR const* UVertexDeltaGraphDebugDataInterface::GetShaderVirtualPath() const
+{
+	return TemplateFilePath;
+}
+
+void UVertexDeltaGraphDebugDataInterface::GetShaderHash(FString& InOutKey) const
+{
+	GetShaderFileHash(TemplateFilePath, EShaderPlatform::SP_PCD3D_SM5).AppendString(InOutKey);
+}
+
+void UVertexDeltaGraphDebugDataInterface::GetHLSL(FString& OutHLSL, FString const& InDataInterfaceName) const
+{
+	TMap<FString, FStringFormatArg> TemplateArgs =
+	{
+		{ TEXT("DataInterfaceName"), InDataInterfaceName },
+	};
+
+	FString TemplateFile;
+	LoadShaderSourceFile(TemplateFilePath, EShaderPlatform::SP_PCD3D_SM5, &TemplateFile, nullptr);
+	OutHLSL += FString::Format(*TemplateFile, TemplateArgs);
+}
+
+UComputeDataProvider* UVertexDeltaGraphDebugDataInterface::CreateDataProvider(TObjectPtr<UObject> InBinding, uint64 InInputMask, uint64 InOutputMask) const
+{
+	UVertexDeltaGraphDebugDataProvider* Provider = NewObject<UVertexDeltaGraphDebugDataProvider>();
+	Provider->DeformerComponent = Cast<UMLDeformerComponent>(InBinding);
+	if (Provider->DeformerComponent)
+	{
+		Provider->DeformerAsset = Provider->DeformerComponent->GetDeformerAsset();
+	}
+	MLDEFORMER_EDITORDATA_ONLY(
+		if (Provider->DeformerAsset != nullptr)
+		{
+			Provider->Init();
+		}
+	,)
+	return Provider;
+}
+
+FComputeDataProviderRenderProxy* UVertexDeltaGraphDebugDataProvider::GetRenderProxy()
+{
+#if WITH_EDITORONLY_DATA
+	if (DeformerComponent && DeformerAsset && DeformerComponent->GetModelInstance() && DeformerComponent->GetModelInstance()->IsValidForDataProvider())
+	{
+		UE::VertexDeltaModel::FVertexDeltaGraphDebugDataProviderProxy* Proxy = new UE::VertexDeltaModel::FVertexDeltaGraphDebugDataProviderProxy(DeformerComponent, DeformerAsset, this);
+		const float SampleTime = DeformerComponent->GetModelInstance()->GetSkeletalMeshComponent()->GetPosition();
+		UMLDeformerModel* Model = DeformerAsset->GetModel();
+		Model->SampleGroundTruthPositions(SampleTime, Proxy->GetGroundTruthPositions());
+		Proxy->HandleZeroGroundTruthPositions();
+		return Proxy;
+	}
+#endif
+
+	// Return default invalid proxy.
+	return new FComputeDataProviderRenderProxy();
+}
+
+#if WITH_EDITORONLY_DATA
+namespace UE::VertexDeltaModel
+{
+	FVertexDeltaGraphDebugDataProviderProxy::FVertexDeltaGraphDebugDataProviderProxy(UMLDeformerComponent* DeformerComponent, UMLDeformerAsset* DeformerAsset, UVertexDeltaGraphDebugDataProvider* InProvider)
+		: FComputeDataProviderRenderProxy()
+	{
+		Provider = InProvider;
+
+		if (DeformerComponent && DeformerAsset)
+		{
+			const UMLDeformerModel* Model = DeformerAsset->GetModel();	
+			const UMLDeformerVizSettings* VizSettings = Model->GetVizSettings();
+			const UMLDeformerModelInstance* ModelInstance = DeformerComponent->GetModelInstance();
+
+			SkeletalMeshObject = ModelInstance->GetSkeletalMeshComponent()->MeshObject;
+			VertexMapBufferSRV = Model->GetVertexMapBuffer().ShaderResourceViewRHI;
+			HeatMapMode = static_cast<int32>(VizSettings->GetHeatMapMode());
+			HeatMapMax = 1.0f / FMath::Max(VizSettings->GetHeatMapMax(), 0.00001f);
+			GroundTruthLerp = VizSettings->GetGroundTruthLerp();
+		}
+	}
+
+	void FVertexDeltaGraphDebugDataProviderProxy::HandleZeroGroundTruthPositions()
+	{
+		if (GroundTruthPositions.Num() == 0)
+		{	
+			// We didn't get valid ground truth vertices.
+			// Make non empty array for later buffer generation.
+			GroundTruthPositions.Add(FVector3f::ZeroVector);
+
+			// Silently disable relevant debug things.
+			if (HeatMapMode == static_cast<int32>(EMLDeformerHeatMapMode::GroundTruth))
+			{
+				HeatMapMode = -1;
+				HeatMapMax = 0.0f;
+				GroundTruthLerp = 0.0f;
+			}
+		}
+	}
+
+	bool FVertexDeltaGraphDebugDataProviderProxy::IsValid(FValidationData const& InValidationData) const
+	{
+		if (InValidationData.ParameterStructSize != sizeof(FVertexDeltaGraphDebugDataInterfaceParameters))
+		{
+			return false;
+		}
+		if (SkeletalMeshObject == nullptr || VertexMapBufferSRV == nullptr)
+		{
+			return false;
+		}
+
+		return true;
+	}
+
+	void FVertexDeltaGraphDebugDataProviderProxy::AllocateResources(FRDGBuilder& GraphBuilder)
+	{
+		GroundTruthBuffer = GraphBuilder.CreateBuffer(FRDGBufferDesc::CreateStructuredDesc(sizeof(float), 3 * GroundTruthPositions.Num()), TEXT("MLDeformer.GroundTruthPositions"));
+		GroundTruthBufferSRV = GraphBuilder.CreateSRV(GroundTruthBuffer);
+		GraphBuilder.QueueBufferUpload(GroundTruthBuffer, GroundTruthPositions.GetData(), sizeof(FVector3f) * GroundTruthPositions.Num(), ERDGInitialDataFlags::None);
+	}
+
+	void FVertexDeltaGraphDebugDataProviderProxy::GatherDispatchData(FDispatchData const& InDispatchData)
+	{
+		const FSkeletalMeshRenderData& SkeletalMeshRenderData = SkeletalMeshObject->GetSkeletalMeshRenderData();
+		const FSkeletalMeshLODRenderData* LodRenderData = SkeletalMeshRenderData.GetPendingFirstLOD(0);
+		const TStridedView<FVertexDeltaGraphDebugDataInterfaceParameters> ParameterArray = MakeStridedParameterView<FVertexDeltaGraphDebugDataInterfaceParameters>(InDispatchData);
+		for (int32 InvocationIndex = 0; InvocationIndex < ParameterArray.Num(); ++InvocationIndex)
+		{
+			const FSkelMeshRenderSection& RenderSection = LodRenderData->RenderSections[InvocationIndex];
+			FVertexDeltaGraphDebugDataInterfaceParameters& Parameters = ParameterArray[InvocationIndex];
+			Parameters.NumVertices = 0;
+			Parameters.InputStreamStart = RenderSection.BaseVertexIndex;
+			Parameters.HeatMapMode = HeatMapMode;
+			Parameters.HeatMapMax = HeatMapMax;
+			Parameters.GroundTruthLerp = GroundTruthLerp;
+			Parameters.GroundTruthBufferSize = GroundTruthPositions.Num();
+			Parameters.PositionGroundTruthBuffer = GroundTruthBufferSRV;
+			Parameters.VertexMapBuffer = VertexMapBufferSRV;
+		}
+	}
+}	// namespace UE::VertexDeltaModel
+#endif // WITH_EDITORONLY_DATA
+
diff --git a/Engine/Plugins/Animation/MLDeformer/VertexDeltaModel/Source/VertexDeltaModel/Private/VertexDeltaModel.cpp b/Engine/Plugins/Animation/MLDeformer/VertexDeltaModel/Source/VertexDeltaModel/Private/VertexDeltaModel.cpp
--- a/Engine/Plugins/Animation/MLDeformer/VertexDeltaModel/Source/VertexDeltaModel/Private/VertexDeltaModel.cpp
+++ b/Engine/Plugins/Animation/MLDeformer/VertexDeltaModel/Source/VertexDeltaModel/Private/VertexDeltaModel.cpp
@@ -2,6 +2,14 @@
 
 #include "VertexDeltaModel.h"
 #include "VertexDeltaModelVizSettings.h"
+#include "VertexDeltaModelInstance.h"
+#include "MLDeformerComponent.h"
+#include "MLDeformerAsset.h"
+#include "NeuralNetwork.h"
+#include "Modules/ModuleManager.h"
+#include "Interfaces/IPluginManager.h"
+#include "Misc/Paths.h"
+#include "ShaderCore.h"
 
 #define LOCTEXT_NAMESPACE "VertexDeltaModel"
 
@@ -13,11 +21,16 @@ namespace UE::VertexDeltaModel
 	class VERTEXDELTAMODEL_API FVertexDeltaModelModule
 		: public IModuleInterface
 	{
+		void StartupModule()
+		{
+			// Register an additional shader path for our shaders used inside the deformer graph system.
+			const FString PluginShaderDir = FPaths::Combine(IPluginManager::Get().FindPlugin(TEXT("VertexDeltaModel"))->GetBaseDir(), TEXT("Shaders"));
+			AddShaderSourceDirectoryMapping(TEXT("/Plugin/VertexDeltaModel"), PluginShaderDir);
+		}
 	};
 }
 IMPLEMENT_MODULE(UE::VertexDeltaModel::FVertexDeltaModelModule, VertexDeltaModel)
 
-
 UVertexDeltaModel::UVertexDeltaModel(const FObjectInitializer& ObjectInitializer)
 	: Super(ObjectInitializer)
 {
@@ -26,9 +39,59 @@ UVertexDeltaModel::UVertexDeltaModel(const FObjectInitializer& ObjectInitializer
 #endif
 }
 
+UMLDeformerModelInstance* UVertexDeltaModel::CreateModelInstance(UMLDeformerComponent* Component)
+{
+	return NewObject<UVertexDeltaModelInstance>(Component);
+}
+
 FString UVertexDeltaModel::GetDefaultDeformerGraphAssetPath() const 
 { 
-	return FString(TEXT("/MLDeformerFramework/Deformers/DG_MLDeformerModel.DG_MLDeformerModel"));
+	return FString(TEXT("/VertexDeltaModel/Deformers/DG_VertexDeltaModel.DG_VertexDeltaModel"));
+}
+
+void UVertexDeltaModel::PostLoad()
+{
+	Super::PostLoad();
+
+	if (NNINetwork)
+	{
+		NNINetwork->SetDeviceType(ENeuralDeviceType::GPU, ENeuralDeviceType::CPU, ENeuralDeviceType::GPU);
+		if (NNINetwork->GetDeviceType() != ENeuralDeviceType::GPU || NNINetwork->GetOutputDeviceType() != ENeuralDeviceType::GPU || NNINetwork->GetInputDeviceType() != ENeuralDeviceType::CPU)
+		{
+			UE_LOG(LogVertexDeltaModel, Error, TEXT("Neural net in MLD Vertex Delta Model '%s' cannot run on the GPU, it will not be active."), *GetDeformerAsset()->GetName());
+		}
+	}
 }
 
+void UVertexDeltaModel::SetNNINetwork(UNeuralNetwork* InNeuralNetwork)
+{
+	GetNeuralNetworkModifyDelegate().Broadcast();
+	NNINetwork = InNeuralNetwork;
+}
+
+UNeuralNetwork* UVertexDeltaModel::GetNNINetwork() const
+{
+	return NNINetwork.Get();
+}
+
+#if WITH_EDITOR
+	void UVertexDeltaModel::UpdateMemoryUsage()
+	{
+		Super::UpdateMemoryUsage();
+
+		// Check if the neural network is on the GPU, if so, count the memory to the GPU and remove it from Main memory, as we added it to 
+		// that already when we did the Model->GetResourceSizeBytes.
+		if (NNINetwork)
+		{
+			if (NNINetwork->GetDeviceType() == ENeuralDeviceType::GPU)
+			{
+				const uint64 NeuralNetSize = static_cast<uint64>(NNINetwork->GetResourceSizeBytes(EResourceSizeMode::Type::EstimatedTotal));
+				GPUMemUsageInBytes += NeuralNetSize;
+				MemUsageInBytes -= NeuralNetSize;
+				CookedMemUsageInBytes -= NeuralNetSize;
+			}
+		}
+	}
+#endif
+
 #undef LOCTEXT_NAMESPACE
diff --git a/Engine/Plugins/Animation/MLDeformer/VertexDeltaModel/Source/VertexDeltaModel/Private/VertexDeltaModelInstance.cpp b/Engine/Plugins/Animation/MLDeformer/VertexDeltaModel/Source/VertexDeltaModel/Private/VertexDeltaModelInstance.cpp
new file mode 100644
--- /dev/null
+++ b/Engine/Plugins/Animation/MLDeformer/VertexDeltaModel/Source/VertexDeltaModel/Private/VertexDeltaModelInstance.cpp
@@ -0,0 +1,151 @@
+// Copyright Epic Games, Inc. All Rights Reserved.
+
+#include "VertexDeltaModelInstance.h"
+#include "VertexDeltaModel.h"
+#include "MLDeformerAsset.h"
+#include "NeuralNetwork.h"
+#include "Components/SkeletalMeshComponent.h"
+
+UVertexDeltaModel* UVertexDeltaModelInstance::GetVertexDeltaModel() const
+{
+	return Cast<UVertexDeltaModel>(Model);
+}
+
+UNeuralNetwork* UVertexDeltaModelInstance::GetNNINetwork() const
+{
+	return GetVertexDeltaModel()->GetNNINetwork();
+}
+
+void UVertexDeltaModelInstance::Release()
+{
+	Super::Release();
+
+	// Destroy the neural network instance.
+	if (GetVertexDeltaModel())
+	{
+		UNeuralNetwork* NNINetwork = GetNNINetwork();
+		if (NNINetwork && NeuralNetworkInferenceHandle != -1)
+		{
+			NNINetwork->DestroyInferenceContext(NeuralNetworkInferenceHandle);
+			NeuralNetworkInferenceHandle = -1;
+		}
+	}
+}
+
+FString UVertexDeltaModelInstance::CheckCompatibility(USkeletalMeshComponent* InSkelMeshComponent, bool bLogIssues)
+{
+	FString ErrorString = Super::CheckCompatibility(InSkelMeshComponent, bLogIssues);
+
+	// Verify the number of inputs versus the expected number of inputs.	
+	const UNeuralNetwork* NNINetwork = GetNNINetwork();
+	if (NNINetwork && NNINetwork->IsLoaded() && Model->GetDeformerAsset())
+	{
+		const int64 NumNeuralNetInputs = NNINetwork->GetInputTensor().Num();
+		const int64 NumDeformerAssetInputs = static_cast<int64>(Model->GetInputInfo()->CalcNumNeuralNetInputs(Model->GetNumFloatsPerBone(), Model->GetNumFloatsPerCurve()));
+		if (NumNeuralNetInputs != NumDeformerAssetInputs)
+		{
+			const FString InputErrorString = "The number of network inputs doesn't match the asset. Please retrain the asset."; 
+			ErrorText += InputErrorString + "\n";
+			if (bLogIssues)
+			{
+				UE_LOG(LogVertexDeltaModel, Error, TEXT("Deformer '%s': %s"), *(Model->GetDeformerAsset()->GetName()), *InputErrorString);
+			}
+		}
+	}
+
+	return ErrorString;
+}
+
+bool UVertexDeltaModelInstance::IsValidForDataProvider() const
+{
+	UNeuralNetwork* NNINetwork = GetNNINetwork();
+	if (NNINetwork == nullptr || !NNINetwork->IsLoaded())
+	{
+		return false;
+	}
+
+	// We expect to run on the GPU when using a data provider for the deformer graph system (Optimus).
+	// Make sure we're actually running the network on the GPU.
+	// Inputs are expected to come from the CPU though.
+	if (NNINetwork->GetDeviceType() != ENeuralDeviceType::GPU || NNINetwork->GetOutputDeviceType() != ENeuralDeviceType::GPU || NNINetwork->GetInputDeviceType() != ENeuralDeviceType::CPU)
+	{
+		return false;
+	}
+
+	return (Model->GetVertexMapBuffer().ShaderResourceViewRHI != nullptr) && (GetNeuralNetworkInferenceHandle() != -1);
+}
+
+void UVertexDeltaModelInstance::Execute(float ModelWeight)
+{
+	TRACE_CPUPROFILER_EVENT_SCOPE(UMLDeformerModelInstance::Execute)
+
+	UNeuralNetwork* NNINetwork = GetNNINetwork();
+	if (NNINetwork == nullptr)
+	{
+		return;
+	}
+
+	// Even if the model needs the GPU it is possible that the hardware does not support GPU evaluation
+	if (NNINetwork->GetDeviceType() == ENeuralDeviceType::GPU)
+	{
+		// NOTE: Inputs still come from the CPU.
+		check(NNINetwork->GetDeviceType() == ENeuralDeviceType::GPU && NNINetwork->GetInputDeviceType() == ENeuralDeviceType::CPU && NNINetwork->GetOutputDeviceType() == ENeuralDeviceType::GPU);
+		ENQUEUE_RENDER_COMMAND(RunNeuralNetwork)
+			(
+				[NNINetwork, Handle = NeuralNetworkInferenceHandle](FRHICommandListImmediate& RHICmdList)
+				{
+					// Output deltas will be available on GPU for DeformerGraph via UMLDeformerDataProvider.
+					FRDGBuilder GraphBuilder(RHICmdList);
+					NNINetwork->Run(GraphBuilder, Handle);
+					GraphBuilder.Execute();
+				}
+		);
+	}
+}
+
+bool UVertexDeltaModelInstance::SetupInputs()
+{
+	TRACE_CPUPROFILER_EVENT_SCOPE(UMLDeformerModelInstance::SetupInputs)
+
+	// Some safety checks.
+	if (Model == nullptr ||
+		SkeletalMeshComponent == nullptr ||
+		SkeletalMeshComponent->GetSkeletalMeshAsset() == nullptr ||
+		!bIsCompatible)
+	{
+		return false;
+	}
+
+	// Get the network and make sure it's loaded.
+	UNeuralNetwork* NNINetwork = GetNNINetwork();
+	if (NNINetwork == nullptr || !NNINetwork->IsLoaded())
+	{
+		return false;
+	}
+
+	// Allocate an inference context if none has already been allocated.
+	if (NeuralNetworkInferenceHandle == -1)
+	{
+		NeuralNetworkInferenceHandle = NNINetwork->CreateInferenceContext();
+		if (NeuralNetworkInferenceHandle == -1)
+		{
+			return false;
+		}
+	}
+
+	// If the neural network expects a different number of inputs, do nothing.
+	const int64 NumNeuralNetInputs = NNINetwork->GetInputTensorForContext(NeuralNetworkInferenceHandle).Num();
+	const int64 NumDeformerAssetInputs = Model->GetInputInfo()->CalcNumNeuralNetInputs(Model->GetNumFloatsPerBone(), Model->GetNumFloatsPerCurve());
+	if (NumNeuralNetInputs != NumDeformerAssetInputs)
+	{
+		return false;
+	}
+
+	// Update and write the input values directly into the input tensor.
+	float* InputDataPointer = static_cast<float*>(NNINetwork->GetInputDataPointerMutableForContext(NeuralNetworkInferenceHandle));
+	const int64 NumFloatsWritten = SetNeuralNetworkInputValues(InputDataPointer, NumNeuralNetInputs);
+	check(NumFloatsWritten == NumNeuralNetInputs);
+
+	return true;
+}
+
diff --git a/Engine/Plugins/Animation/MLDeformer/VertexDeltaModel/Source/VertexDeltaModel/Public/VertexDeltaGraphDataInterface.h b/Engine/Plugins/Animation/MLDeformer/VertexDeltaModel/Source/VertexDeltaModel/Public/VertexDeltaGraphDataInterface.h
new file mode 100644
--- /dev/null
+++ b/Engine/Plugins/Animation/MLDeformer/VertexDeltaModel/Source/VertexDeltaModel/Public/VertexDeltaGraphDataInterface.h
@@ -0,0 +1,98 @@
+// Copyright Epic Games, Inc. All Rights Reserved.
+
+#pragma once
+
+#include "CoreMinimal.h"
+#include "OptimusComputeDataInterface.h"
+#include "ComputeFramework/ComputeDataProvider.h"
+#include "Rendering/SkeletalMeshRenderData.h"
+#include "Rendering/SkeletalMeshLODRenderData.h"
+#include "SkeletalRenderPublic.h"
+#include "VertexDeltaGraphDataInterface.generated.h"
+
+class FRDGBuffer;
+class FRDGBufferSRV;
+class FRHIShaderResourceView;
+class FSkeletalMeshObject;
+class UMLDeformerComponent;
+class UNeuralNetwork;
+class USkeletalMeshComponent;
+class UMLDeformerModel;
+
+#if WITH_EDITORONLY_DATA
+	#define MLDEFORMER_EDITORDATA_ONLY(Statement, ElseStatement) Statement
+#else
+	#define MLDEFORMER_EDITORDATA_ONLY(X, ElseStatement) ElseStatement
+#endif
+
+/** Compute Framework Data Interface for MLDefomer data. */
+UCLASS(Category = ComputeFramework)
+class VERTEXDELTAMODEL_API UVertexDeltaGraphDataInterface
+	: public UOptimusComputeDataInterface
+{
+	GENERATED_BODY()
+
+public:
+	// UOptimusComputeDataInterface overrides.
+	virtual FString GetDisplayName() const override;
+	virtual TArray<FOptimusCDIPinDefinition> GetPinDefinitions() const override;
+	virtual TSubclassOf<UActorComponent> GetRequiredComponentClass() const override;
+	// ~END UOptimusComputeDataInterface overrides.
+
+	// UComputeDataInterface overrides.
+	TCHAR const* GetClassName() const override { return TEXT("VertexDeltaModelData"); }
+	virtual void GetSupportedInputs(TArray<FShaderFunctionDefinition>& OutFunctions) const override;
+	virtual void GetShaderParameters(TCHAR const* UID, FShaderParametersMetadataBuilder& InOutBuilder, FShaderParametersMetadataAllocations& InOutAllocations) const override;
+	virtual TCHAR const* GetShaderVirtualPath() const override;
+	virtual void GetShaderHash(FString& InOutKey) const override;
+	virtual void GetHLSL(FString& OutHLSL, FString const& InDataInterfaceName) const override;
+	virtual UComputeDataProvider* CreateDataProvider(TObjectPtr<UObject> InBinding, uint64 InInputMask, uint64 InOutputMask) const override;
+	// ~END UComputeDataInterface overrides.
+
+private:
+	static TCHAR const* TemplateFilePath;
+};
+
+/** Compute Framework Data Provider for MLDeformer data. */
+UCLASS(BlueprintType, EditInlineNew, Category = ComputeFramework)
+class VERTEXDELTAMODEL_API UVertexDeltaGraphDataProvider
+	: public UComputeDataProvider
+{
+	GENERATED_BODY()
+
+public:
+	/** The deformer component that this data provider works on. */
+	UPROPERTY(EditAnywhere, BlueprintReadWrite, Category = Binding)
+	TObjectPtr<UMLDeformerComponent> DeformerComponent = nullptr;
+
+	// UComputeDataProvider overrides.
+	virtual FComputeDataProviderRenderProxy* GetRenderProxy() override;
+	// ~END UComputeDataProvider overrides.
+};
+
+namespace UE::VertexDeltaModel
+{
+	/** Compute Framework Data Provider Proxy for MLDeformer data. */
+	class VERTEXDELTAMODEL_API FVertexDeltaGraphDataProviderProxy
+		: public FComputeDataProviderRenderProxy
+	{
+	public:
+		FVertexDeltaGraphDataProviderProxy(UMLDeformerComponent* DeformerComponent);
+
+		// FComputeDataProviderRenderProxy overrides.
+		bool IsValid(FValidationData const& InValidationData) const override;
+		void AllocateResources(FRDGBuilder& GraphBuilder) override;
+		void GatherDispatchData(FDispatchData const& InDispatchData) override;
+		// ~END FComputeDataProviderRenderProxy overrides.
+
+	protected:
+		FSkeletalMeshObject* SkeletalMeshObject = nullptr;
+		TObjectPtr<UNeuralNetwork> NeuralNetwork = nullptr;
+		FRHIShaderResourceView* VertexMapBufferSRV = nullptr;
+		FRDGBuffer* Buffer = nullptr;
+		FRDGBufferSRV* BufferSRV = nullptr;
+		int32 NeuralNetworkInferenceHandle = -1;
+		float Weight = 1.0f;
+		bool bCanRunNeuralNet = false;
+	};
+}	// namespace UE::VertexDeltaModel
diff --git a/Engine/Plugins/Animation/MLDeformer/VertexDeltaModel/Source/VertexDeltaModel/Public/VertexDeltaGraphDebugDataInterface.h b/Engine/Plugins/Animation/MLDeformer/VertexDeltaModel/Source/VertexDeltaModel/Public/VertexDeltaGraphDebugDataInterface.h
new file mode 100644
--- /dev/null
+++ b/Engine/Plugins/Animation/MLDeformer/VertexDeltaModel/Source/VertexDeltaModel/Public/VertexDeltaGraphDebugDataInterface.h
@@ -0,0 +1,104 @@
+// Copyright Epic Games, Inc. All Rights Reserved.
+
+#pragma once
+
+#include "CoreMinimal.h"
+#include "MLDeformerModel.h"
+#include "VertexDeltaGraphDataInterface.h"
+#include "OptimusComputeDataInterface.h"
+#include "ComputeFramework/ComputeDataProvider.h"
+#include "VertexDeltaGraphDebugDataInterface.generated.h"
+
+class FRDGBuffer;
+class FRDGBufferSRV;
+class FRHIShaderResourceView;
+class FSkeletalMeshObject;
+class UMLDeformerAsset;
+class USkeletalMeshComponent;
+
+/** 
+ * Compute Framework Data Interface for MLDeformer debugging data. 
+ * This interfaces to editor only data, and so will only give valid results in that context.
+ */
+UCLASS(Category = ComputeFramework)
+class VERTEXDELTAMODEL_API UVertexDeltaGraphDebugDataInterface
+	: public UOptimusComputeDataInterface
+{
+	GENERATED_BODY()
+
+public:
+	// UOptimusComputeDataInterface overrides.
+	virtual FString GetDisplayName() const override;
+	virtual TArray<FOptimusCDIPinDefinition> GetPinDefinitions() const override;
+	virtual TSubclassOf<UActorComponent> GetRequiredComponentClass() const override;
+	// ~END UOptimusComputeDataInterface overrides.
+
+	// UComputeDataInterface overrides.
+	TCHAR const* GetClassName() const override { return TEXT("VertexDeltaModelDebugData"); }
+	virtual void GetSupportedInputs(TArray<FShaderFunctionDefinition>& OutFunctions) const override;
+	virtual void GetShaderParameters(TCHAR const* UID, FShaderParametersMetadataBuilder& OutBuilder, FShaderParametersMetadataAllocations& InOutAllocations) const override;
+	virtual TCHAR const* GetShaderVirtualPath() const override;
+	virtual void GetShaderHash(FString& InOutKey) const override;
+	virtual void GetHLSL(FString& OutHLSL, FString const& InDataInterfaceName) const override;
+	virtual UComputeDataProvider* CreateDataProvider(TObjectPtr<UObject> InBinding, uint64 InInputMask, uint64 InOutputMask) const override;
+	// ~END UComputeDataInterface overrides.
+
+private:
+	static TCHAR const* TemplateFilePath;
+};
+
+/** Compute Framework Data Provider for MLDeformer debugging data. */
+UCLASS(BlueprintType, EditInlineNew, Category = ComputeFramework)
+class VERTEXDELTAMODEL_API UVertexDeltaGraphDebugDataProvider
+	: public UComputeDataProvider
+{
+	GENERATED_BODY()
+
+public:
+	virtual void Init() {};
+
+	// UComputeDataProvider overrides.
+	virtual FComputeDataProviderRenderProxy* GetRenderProxy() override;
+	// ~END UComputeDataProvider overrides.
+
+public:
+	UPROPERTY(EditAnywhere, BlueprintReadWrite, Category = Binding)
+	TObjectPtr<UMLDeformerComponent> DeformerComponent = nullptr;
+
+	UPROPERTY(EditAnywhere, BlueprintReadWrite, Category = Binding)
+	TObjectPtr<UMLDeformerAsset> DeformerAsset = nullptr;
+};
+
+/** Compute Framework Data Provider Proxy for MLDeformer debugging data. */
+#if WITH_EDITORONLY_DATA
+namespace UE::VertexDeltaModel
+{
+	class VERTEXDELTAMODEL_API FVertexDeltaGraphDebugDataProviderProxy
+		: public FComputeDataProviderRenderProxy
+	{
+	public:
+		FVertexDeltaGraphDebugDataProviderProxy(UMLDeformerComponent* DeformerComponent, UMLDeformerAsset* DeformerAsset, UVertexDeltaGraphDebugDataProvider* InProvider);
+
+		virtual void HandleZeroGroundTruthPositions();
+
+		// FComputeDataProviderRenderProxy overrides.
+		bool IsValid(FValidationData const& InValidationData) const override;
+		void AllocateResources(FRDGBuilder& GraphBuilder) override;
+		void GatherDispatchData(FDispatchData const& InDispatchData) override;
+		// ~END FComputeDataProviderRenderProxy overrides.
+
+		TArray<FVector3f>& GetGroundTruthPositions() { return GroundTruthPositions; }
+
+	protected:
+		TObjectPtr<UVertexDeltaGraphDebugDataProvider> Provider = nullptr;
+		FSkeletalMeshObject* SkeletalMeshObject = nullptr;
+		TArray<FVector3f> GroundTruthPositions;
+		FRHIShaderResourceView* VertexMapBufferSRV = nullptr;
+		FRDGBuffer* GroundTruthBuffer = nullptr;
+		FRDGBufferSRV* GroundTruthBufferSRV = nullptr;
+		int32 HeatMapMode = 0;
+		float HeatMapMax = 0.0f;
+		float GroundTruthLerp = 0.0f;
+	};
+}	// namespace UE::VertexDeltaModel
+#endif
diff --git a/Engine/Plugins/Animation/MLDeformer/VertexDeltaModel/Source/VertexDeltaModel/Public/VertexDeltaModel.h b/Engine/Plugins/Animation/MLDeformer/VertexDeltaModel/Source/VertexDeltaModel/Public/VertexDeltaModel.h
--- a/Engine/Plugins/Animation/MLDeformer/VertexDeltaModel/Source/VertexDeltaModel/Public/VertexDeltaModel.h
+++ b/Engine/Plugins/Animation/MLDeformer/VertexDeltaModel/Source/VertexDeltaModel/Public/VertexDeltaModel.h
@@ -9,6 +9,8 @@
 
 VERTEXDELTAMODEL_API DECLARE_LOG_CATEGORY_EXTERN(LogVertexDeltaModel, Log, All);
 
+class UNeuralNetwork;
+
 /** 
  * The vertex delta model, which uses a GPU based neural network.
  * This model acts more as an example of how to implement a model that only uses the GPU.
@@ -25,20 +27,38 @@ class VERTEXDELTAMODEL_API UVertexDeltaModel
 public:
 	UVertexDeltaModel(const FObjectInitializer& ObjectInitializer);
 
+	// UObject overrides.
+	void PostLoad() override;
+	// ~END UObject overrides.
+
 	// UMLDeformerModel overrides.
 	virtual FString GetDisplayName() const override			{ return "Vertex Delta Model"; }
+	virtual UMLDeformerModelInstance* CreateModelInstance(UMLDeformerComponent* Component) override;
 	virtual bool IsNeuralNetworkOnGPU() const override		{ return true; }	// GPU neural network.
 	virtual FString GetDefaultDeformerGraphAssetPath() const override;
+#if WITH_EDITOR
+	virtual void UpdateMemoryUsage() override;
+#endif
 	// ~END UMLDeformerModel overrides.
 
+	UNeuralNetwork* GetNNINetwork() const;
+	void SetNNINetwork(UNeuralNetwork* InNeuralNetwork);
+
 #if WITH_EDITORONLY_DATA
 	int32 GetNumHiddenLayers() const						{ return NumHiddenLayers; }
 	int32 GetNumNeuronsPerLayer() const						{ return NumNeuronsPerLayer; }
 	int32 GetNumIterations() const							{ return NumIterations; }
 	int32 GetBatchSize() const								{ return BatchSize; }
 	float GetLearningRate() const							{ return LearningRate; }
+#endif
+
 
 public:
+	/** The NNI neural network. */
+	UPROPERTY()
+	TObjectPtr<UNeuralNetwork> NNINetwork;
+
+#if WITH_EDITORONLY_DATA
 	/** The number of hidden layers that the neural network model will have.\nHigher numbers will slow down performance but can deal with more complex deformations. */
 	UPROPERTY(EditAnywhere, BlueprintReadWrite, Category = "Training Settings", meta = (ClampMin = "1", ClampMax = "10"))
 	int32 NumHiddenLayers = 3;
diff --git a/Engine/Plugins/Animation/MLDeformer/VertexDeltaModel/Source/VertexDeltaModel/Public/VertexDeltaModelInstance.h b/Engine/Plugins/Animation/MLDeformer/VertexDeltaModel/Source/VertexDeltaModel/Public/VertexDeltaModelInstance.h
new file mode 100644
--- /dev/null
+++ b/Engine/Plugins/Animation/MLDeformer/VertexDeltaModel/Source/VertexDeltaModel/Public/VertexDeltaModelInstance.h
@@ -0,0 +1,28 @@
+// Copyright Epic Games, Inc. All Rights Reserved.
+
+#pragma once
+
+#include "MLDeformerModelInstance.h"
+#include "VertexDeltaModelInstance.generated.h"
+
+class UVertexDeltaModel;
+class UNeuralNetwork;
+
+UCLASS()
+class VERTEXDELTAMODEL_API UVertexDeltaModelInstance
+	: public UMLDeformerModelInstance
+{
+	GENERATED_BODY()
+
+public:
+	// UMLDeformerModelInstance overrides.
+	virtual void Release() override;
+	virtual FString CheckCompatibility(USkeletalMeshComponent* InSkelMeshComponent, bool bLogIssues) override;
+	virtual bool IsValidForDataProvider() const override;
+	virtual void Execute(float ModelWeight) override;
+	virtual bool SetupInputs() override;
+	// ~END UMLDeformerModelInstance overrides.
+
+	UVertexDeltaModel* GetVertexDeltaModel() const;
+	UNeuralNetwork* GetNNINetwork() const;
+};
diff --git a/Engine/Plugins/Animation/MLDeformer/VertexDeltaModel/Source/VertexDeltaModelEditor/Private/VertexDeltaEditorModel.cpp b/Engine/Plugins/Animation/MLDeformer/VertexDeltaModel/Source/VertexDeltaModelEditor/Private/VertexDeltaEditorModel.cpp
--- a/Engine/Plugins/Animation/MLDeformer/VertexDeltaModel/Source/VertexDeltaModelEditor/Private/VertexDeltaEditorModel.cpp
+++ b/Engine/Plugins/Animation/MLDeformer/VertexDeltaModel/Source/VertexDeltaModelEditor/Private/VertexDeltaEditorModel.cpp
@@ -4,18 +4,8 @@
 #include "VertexDeltaModel.h"
 #include "VertexDeltaModelVizSettings.h"
 #include "VertexDeltaTrainingModel.h"
-#include "MLDeformerGeomCacheActor.h"
-#include "MLDeformerEditorToolkit.h"
-#include "MLDeformerEditorStyle.h"
-#include "MLDeformerInputInfo.h"
-#include "MLDeformerGeomCacheSampler.h"
-#include "MLDeformerGeomCacheHelpers.h"
-#include "GeometryCache.h"
-#include "GeometryCacheComponent.h"
-#include "Animation/DebugSkelMeshComponent.h"
-#include "Rendering/SkeletalMeshModel.h"
-#include "Rendering/SkeletalMeshLODModel.h"
-#include "Engine/SkeletalMesh.h"
+#include "MLDeformerAsset.h"
+#include "NeuralNetwork.h"
 
 #define LOCTEXT_NAMESPACE "VertexDeltaEditorModel"
 
@@ -28,10 +18,68 @@ namespace UE::VertexDeltaModel
 		return new FVertexDeltaEditorModel();
 	}
 
+	UVertexDeltaModel* FVertexDeltaEditorModel::GetVertexDeltaModel() const
+	{
+		return Cast<UVertexDeltaModel>(Model);
+	}
+
 	ETrainingResult FVertexDeltaEditorModel::Train()
 	{
 		return TrainModel<UVertexDeltaTrainingModel>(this);
 	}
+
+	UNeuralNetwork* FVertexDeltaEditorModel::LoadNeuralNetworkFromOnnx(const FString& Filename) const
+	{
+		const FString OnnxFile = FPaths::ConvertRelativePathToFull(Filename);
+		if (FPaths::FileExists(OnnxFile))
+		{
+			UE_LOG(LogVertexDeltaModel, Display, TEXT("Loading Onnx file '%s'..."), *OnnxFile);
+			UNeuralNetwork* Result = NewObject<UNeuralNetwork>(Model, UNeuralNetwork::StaticClass());		
+			if (Result->Load(OnnxFile))
+			{
+				Result->SetDeviceType(ENeuralDeviceType::GPU, ENeuralDeviceType::CPU, ENeuralDeviceType::GPU);	
+				if (Result->GetDeviceType() != ENeuralDeviceType::GPU || Result->GetOutputDeviceType() != ENeuralDeviceType::GPU || Result->GetInputDeviceType() != ENeuralDeviceType::CPU)
+				{
+					UE_LOG(LogVertexDeltaModel, Error, TEXT("Neural net in ML Deformer '%s' cannot run on the GPU, it will not be active."), *Model->GetDeformerAsset()->GetName());
+				}
+				UE_LOG(LogVertexDeltaModel, Display, TEXT("Successfully loaded Onnx file '%s'..."), *OnnxFile);
+				return Result;
+			}
+			else
+			{
+				UE_LOG(LogVertexDeltaModel, Error, TEXT("Failed to load Onnx file '%s'"), *OnnxFile);
+			}
+		}
+		else
+		{
+			UE_LOG(LogVertexDeltaModel, Error, TEXT("Onnx file '%s' does not exist!"), *OnnxFile);
+		}
+
+		return nullptr;
+	}
+
+	bool FVertexDeltaEditorModel::LoadTrainedNetwork() const
+	{
+		const FString OnnxFile = GetTrainedNetworkOnnxFile();
+		UNeuralNetwork* Network = LoadNeuralNetworkFromOnnx(OnnxFile);
+		if (Network)
+		{
+			GetVertexDeltaModel()->SetNNINetwork(Network);
+			return true;
+		}
+
+		return false;
+	}
+
+	bool FVertexDeltaEditorModel::IsTrained() const
+	{
+		return (GetVertexDeltaModel()->GetNNINetwork() != nullptr);
+	}
+
+	FString FVertexDeltaEditorModel::GetHeatMapDeformerGraphPath() const
+	{
+		return FString(TEXT("/VertexDeltaModel/Deformers/DG_VertexDeltaModel_HeatMap.DG_VertexDeltaModel_HeatMap"));
+	}
 }	// namespace UE::VertexDeltaModel
 
 #undef LOCTEXT_NAMESPACE
diff --git a/Engine/Plugins/Animation/MLDeformer/VertexDeltaModel/Source/VertexDeltaModelEditor/Public/VertexDeltaEditorModel.h b/Engine/Plugins/Animation/MLDeformer/VertexDeltaModel/Source/VertexDeltaModelEditor/Public/VertexDeltaEditorModel.h
--- a/Engine/Plugins/Animation/MLDeformer/VertexDeltaModel/Source/VertexDeltaModelEditor/Public/VertexDeltaEditorModel.h
+++ b/Engine/Plugins/Animation/MLDeformer/VertexDeltaModel/Source/VertexDeltaModelEditor/Public/VertexDeltaEditorModel.h
@@ -5,6 +5,9 @@
 #include "CoreMinimal.h"
 #include "MLDeformerGeomCacheEditorModel.h"
 
+class UNeuralNetwork;
+class UVertexDeltaModel;
+
 namespace UE::VertexDeltaModel
 {
 	using namespace UE::MLDeformer;
@@ -25,6 +28,12 @@ namespace UE::VertexDeltaModel
 	
 		// FMLDeformerEditorModel overrides.
 		virtual ETrainingResult Train() override;
+		virtual bool LoadTrainedNetwork() const override;
+		virtual bool IsTrained() const override;
+		virtual FString GetHeatMapDeformerGraphPath() const override;
 		// ~END FMLDeformerEditorModel overrides.
+
+		UNeuralNetwork* LoadNeuralNetworkFromOnnx(const FString& Filename) const;
+		UVertexDeltaModel* GetVertexDeltaModel() const;
 	};
 }	// namespace UE::VertexDeltaModel
diff --git a/Engine/Plugins/Animation/MLDeformer/VertexDeltaModel/VertexDeltaModel.uplugin b/Engine/Plugins/Animation/MLDeformer/VertexDeltaModel/VertexDeltaModel.uplugin
--- a/Engine/Plugins/Animation/MLDeformer/VertexDeltaModel/VertexDeltaModel.uplugin
+++ b/Engine/Plugins/Animation/MLDeformer/VertexDeltaModel/VertexDeltaModel.uplugin
@@ -1,7 +1,7 @@
 {
 	"FileVersion": 3,
 	"Version": 1,
-	"VersionName": "0.21",
+	"VersionName": "0.3",
 	"FriendlyName": "ML Deformer Vertex Delta Model",
 	"Description": "Vertex Delta Model for the ML Deformer Framework",
 	"Category": "Animation",
