[MLDeformer] Fixed a crash when you would create a new actor blueprint, added a skelmesh and ml deformer component, then set a neural morph model as mld asset, and in the skelmesh component would select the vertex delta model mesh deformer.

diff --git a/Engine/Plugins/Animation/MLDeformer/VertexDeltaModel/Source/VertexDeltaModel/Private/VertexDeltaGraphDataInterface.cpp b/Engine/Plugins/Animation/MLDeformer/VertexDeltaModel/Source/VertexDeltaModel/Private/VertexDeltaGraphDataInterface.cpp
--- a/Engine/Plugins/Animation/MLDeformer/VertexDeltaModel/Source/VertexDeltaModel/Private/VertexDeltaGraphDataInterface.cpp
+++ b/Engine/Plugins/Animation/MLDeformer/VertexDeltaModel/Source/VertexDeltaModel/Private/VertexDeltaGraphDataInterface.cpp
@@ -106,9 +106,9 @@ namespace UE::VertexDeltaModel
 		const UMLDeformerAsset* DeformerAsset = DeformerComponent != nullptr ? DeformerComponent->GetDeformerAsset() : nullptr;
 		const UMLDeformerModel* Model = DeformerAsset != nullptr ? DeformerAsset->GetModel() : nullptr;
 		const UMLDeformerModelInstance* ModelInstance = DeformerComponent != nullptr ? DeformerComponent->GetModelInstance() : nullptr;
-		if (Model != nullptr && ModelInstance != nullptr)
+		const UVertexDeltaModel* VertexDeltaModel = Cast<UVertexDeltaModel>(Model);
+		if (VertexDeltaModel != nullptr && ModelInstance != nullptr)
 		{
-			const UVertexDeltaModel* VertexDeltaModel = Cast<UVertexDeltaModel>(Model);
 			SkeletalMeshObject = ModelInstance->GetSkeletalMeshComponent()->MeshObject;
 			NeuralNetwork = VertexDeltaModel->GetNNINetwork();
 			NeuralNetworkInferenceHandle = ModelInstance->GetNeuralNetworkInferenceHandle();
