[MLDeformer] Fix a static analysis error.

diff --git a/Engine/Plugins/Animation/MLDeformer/VertexDeltaModel/Source/VertexDeltaModel/Private/VertexDeltaGraphDataInterface.cpp b/Engine/Plugins/Animation/MLDeformer/VertexDeltaModel/Source/VertexDeltaModel/Private/VertexDeltaGraphDataInterface.cpp
--- a/Engine/Plugins/Animation/MLDeformer/VertexDeltaModel/Source/VertexDeltaModel/Private/VertexDeltaGraphDataInterface.cpp
+++ b/Engine/Plugins/Animation/MLDeformer/VertexDeltaModel/Source/VertexDeltaModel/Private/VertexDeltaGraphDataInterface.cpp
@@ -107,7 +107,7 @@ namespace UE::VertexDeltaModel
 		const UMLDeformerModel* Model = DeformerAsset != nullptr ? DeformerAsset->GetModel() : nullptr;
 		const UMLDeformerModelInstance* ModelInstance = DeformerComponent != nullptr ? DeformerComponent->GetModelInstance() : nullptr;
 		const UVertexDeltaModel* VertexDeltaModel = Cast<UVertexDeltaModel>(Model);
-		if (VertexDeltaModel != nullptr && ModelInstance != nullptr)
+		if (Model && VertexDeltaModel && ModelInstance)
 		{
 			SkeletalMeshObject = ModelInstance->GetSkeletalMeshComponent()->MeshObject;
 			NeuralNetwork = VertexDeltaModel->GetNNINetwork();
