Remove NNI from VertexDeltaModel and NearestNeighborModel

diff --git a/Engine/Plugins/Animation/MLDeformer/MLDeformerFramework/Source/MLDeformerFramework/Private/MLDeformerModelInstance.cpp b/Engine/Plugins/Animation/MLDeformer/MLDeformerFramework/Source/MLDeformerFramework/Private/MLDeformerModelInstance.cpp
--- a/Engine/Plugins/Animation/MLDeformer/MLDeformerFramework/Source/MLDeformerFramework/Private/MLDeformerModelInstance.cpp
+++ b/Engine/Plugins/Animation/MLDeformer/MLDeformerFramework/Source/MLDeformerFramework/Private/MLDeformerModelInstance.cpp
@@ -56,7 +56,7 @@ UMLDeformerComponent* UMLDeformerModelInstance::GetMLDeformerComponent() const
 
 int32 UMLDeformerModelInstance::GetNeuralNetworkInferenceHandle() const
 { 
-	return NeuralNetworkInferenceHandle;
+	return -1;
 }
 
 void UMLDeformerModelInstance::SetHasPostInitialized(bool bHasInitialized)
diff --git a/Engine/Plugins/Animation/MLDeformer/MLDeformerFramework/Source/MLDeformerFramework/Public/MLDeformerModel.h b/Engine/Plugins/Animation/MLDeformer/MLDeformerFramework/Source/MLDeformerFramework/Public/MLDeformerModel.h
--- a/Engine/Plugins/Animation/MLDeformer/MLDeformerFramework/Source/MLDeformerFramework/Public/MLDeformerModel.h
+++ b/Engine/Plugins/Animation/MLDeformer/MLDeformerFramework/Source/MLDeformerFramework/Public/MLDeformerModel.h
@@ -19,7 +19,6 @@ class UMLDeformerVizSettings;
 class UMLDeformerModelInstance;
 class UMLDeformerComponent;
 class UMLDeformerInputInfo;
-class UNeuralNetwork;
 
 /** The channel to get the mask data from. */
 UENUM()
@@ -186,11 +185,6 @@ public:
 	 */
 	void SetSkeletalMesh(USkeletalMesh* SkelMesh)			{ SkeletalMesh = SkelMesh; }
 
-	UE_DEPRECATED(5.2, "This method will be deleted. We are moving neural networks into the derived models instead. This method will not do anything.")
-	virtual void SetNeuralNetwork(UNeuralNetwork* InNeuralNetwork) {}
-
-	UE_DEPRECATED(5.2, "This method will be deleted. We are moving neural networks into the derived models instead. A nullptr will be returned.")
-	UNeuralNetwork* GetNeuralNetwork() const				{ return nullptr; }
 
 #if WITH_EDITORONLY_DATA
 	/**
diff --git a/Engine/Plugins/Animation/MLDeformer/MLDeformerFramework/Source/MLDeformerFramework/Public/MLDeformerModelInstance.h b/Engine/Plugins/Animation/MLDeformer/MLDeformerFramework/Source/MLDeformerFramework/Public/MLDeformerModelInstance.h
--- a/Engine/Plugins/Animation/MLDeformer/MLDeformerFramework/Source/MLDeformerFramework/Public/MLDeformerModelInstance.h
+++ b/Engine/Plugins/Animation/MLDeformer/MLDeformerFramework/Source/MLDeformerFramework/Public/MLDeformerModelInstance.h
@@ -116,6 +116,7 @@ public:
 	 * Get the neural network inference handle.
 	 * @return Returns the handle. This is -1 for an invalid handle.
 	 */
+	UE_DEPRECATED(5.3, "The NNI API has been removed - please use NNE instead")
 	int32 GetNeuralNetworkInferenceHandle() const;
 
 	/** Set whether we had a successful call to UMLDeformerModel::PostMLDeformerComponentInit. */
@@ -226,9 +227,6 @@ protected:
 	/** The compatibility error text, in case bIsCompatible is false. */
 	FString ErrorText;
 
-	/** Inference handle for neural network inference. */
-	int32 NeuralNetworkInferenceHandle = -1;
-
 	/** Are the deformer asset and the used skeletal mesh component compatible? */
 	bool bIsCompatible = false;
 
diff --git a/Engine/Plugins/Animation/MLDeformer/MLDeformerFramework/Source/MLDeformerFrameworkEditor/Public/MLDeformerEditorModel.h b/Engine/Plugins/Animation/MLDeformer/MLDeformerFramework/Source/MLDeformerFrameworkEditor/Public/MLDeformerEditorModel.h
--- a/Engine/Plugins/Animation/MLDeformer/MLDeformerFramework/Source/MLDeformerFrameworkEditor/Public/MLDeformerEditorModel.h
+++ b/Engine/Plugins/Animation/MLDeformer/MLDeformerFramework/Source/MLDeformerFrameworkEditor/Public/MLDeformerEditorModel.h
@@ -29,7 +29,6 @@ class USkeletalMesh;
 class UAnimSequence;
 class UMaterial;
 class UGeometryCache;
-class UNeuralNetwork;
 class UMorphTarget;
 class FMorphTargetVertexInfoBuffers;
 
@@ -515,9 +514,6 @@ namespace UE::MLDeformer
 		 */
 		virtual void UpdateDeformerGraph();
 
-		UE_DEPRECATED(5.2, "This method will be deleted and will just return a nullptr right now.")
-		UNeuralNetwork* LoadNeuralNetworkFromOnnx(const FString& Filename) const { return nullptr; }
-
 		/**
 		 * Sample the vertex deltas between the training base model and target model.
 		 * This will initialize the sampler if needed and set the sampler space to post-skinning deltas, and then samples them using the sampler.
diff --git a/Engine/Plugins/Animation/MLDeformer/VertexDeltaModel/Source/VertexDeltaModel/Private/Tests/VertexDeltaModelTest.cpp b/Engine/Plugins/Animation/MLDeformer/VertexDeltaModel/Source/VertexDeltaModel/Private/Tests/VertexDeltaModelTest.cpp
--- a/Engine/Plugins/Animation/MLDeformer/VertexDeltaModel/Source/VertexDeltaModel/Private/Tests/VertexDeltaModelTest.cpp
+++ b/Engine/Plugins/Animation/MLDeformer/VertexDeltaModel/Source/VertexDeltaModel/Private/Tests/VertexDeltaModelTest.cpp
@@ -14,7 +14,6 @@
 #include "UObject/GCObjectScopeGuard.h"
 #include "PreviewScene.h"
 #include "Components/SkeletalMeshComponent.h"
-#include "NeuralNetwork.h"
 #include "GeometryCache.h"
 #include "Engine/SkeletalMesh.h"
 #include "Rendering/SkeletalMeshModel.h"
diff --git a/Engine/Plugins/Animation/MLDeformer/VertexDeltaModel/Source/VertexDeltaModel/Private/VertexDeltaGraphDataInterface.cpp b/Engine/Plugins/Animation/MLDeformer/VertexDeltaModel/Source/VertexDeltaModel/Private/VertexDeltaGraphDataInterface.cpp
--- a/Engine/Plugins/Animation/MLDeformer/VertexDeltaModel/Source/VertexDeltaModel/Private/VertexDeltaGraphDataInterface.cpp
+++ b/Engine/Plugins/Animation/MLDeformer/VertexDeltaModel/Source/VertexDeltaModel/Private/VertexDeltaGraphDataInterface.cpp
@@ -8,7 +8,8 @@
 #include "MLDeformerModelInstance.h"
 #include "MLDeformerAsset.h"
 #include "MLDeformerComponent.h"
-#include "NeuralNetwork.h"
+
+#include "VertexDeltaModelInstance.h"
 #include "OptimusDataDomain.h"
 #include "RenderGraphBuilder.h"
 #include "RenderGraphUtils.h"
@@ -105,16 +106,24 @@ namespace UE::VertexDeltaModel
 
 		const UMLDeformerAsset* DeformerAsset = DeformerComponent != nullptr ? DeformerComponent->GetDeformerAsset() : nullptr;
 		const UMLDeformerModel* Model = DeformerAsset != nullptr ? DeformerAsset->GetModel() : nullptr;
-		const UMLDeformerModelInstance* ModelInstance = DeformerComponent != nullptr ? DeformerComponent->GetModelInstance() : nullptr;
+		UMLDeformerModelInstance* ModelInstance = DeformerComponent != nullptr ? DeformerComponent->GetModelInstance() : nullptr;
 		const UVertexDeltaModel* VertexDeltaModel = Cast<UVertexDeltaModel>(Model);
+		
 		if (Model && VertexDeltaModel && ModelInstance)
 		{
-			SkeletalMeshObject = ModelInstance->GetSkeletalMeshComponent()->MeshObject;
-			NeuralNetwork = VertexDeltaModel->GetNNINetwork();
-			NeuralNetworkInferenceHandle = ModelInstance->GetNeuralNetworkInferenceHandle();
-			bCanRunNeuralNet = ModelInstance->IsCompatible() && Model->IsNeuralNetworkOnGPU() && DeformerComponent->GetModelInstance()->IsValidForDataProvider();
-			Weight = DeformerComponent->GetWeight();
-			VertexMapBufferSRV = VertexDeltaModel->GetVertexMapBuffer().ShaderResourceViewRHI;
+			VertexDeltaModelInstance = Cast<UVertexDeltaModelInstance>(ModelInstance);
+			if (VertexDeltaModelInstance)
+			{
+				SkeletalMeshObject = ModelInstance->GetSkeletalMeshComponent()->MeshObject;
+				NNEModelRDG = VertexDeltaModelInstance->GetNNEModelRDG();
+
+				if (NNEModelRDG)
+				{
+					const FVertexMapBuffer& VertexMapBuffer = VertexDeltaModel->GetVertexMapBuffer();
+					Weight = DeformerComponent->GetWeight();
+					VertexMapBufferSRV = VertexDeltaModel->GetVertexMapBuffer().ShaderResourceViewRHI;
+				}
+			}
 		}
 	}
 
@@ -125,7 +134,7 @@ namespace UE::VertexDeltaModel
 			return false;
 		}
 
-		if (!bCanRunNeuralNet || SkeletalMeshObject == nullptr || NeuralNetwork == nullptr || VertexMapBufferSRV == nullptr)
+		if (!bCanRunNeuralNet || SkeletalMeshObject == nullptr || NNEModelRDG == nullptr || VertexMapBufferSRV == nullptr)
 		{
 			return false;
 		}
@@ -133,12 +142,6 @@ namespace UE::VertexDeltaModel
 		return true;
 	}
 
-	void FVertexDeltaGraphDataProviderProxy::AllocateResources(FRDGBuilder& GraphBuilder)
-	{
-		Buffer = GraphBuilder.RegisterExternalBuffer(NeuralNetwork->GetOutputTensorForContext(NeuralNetworkInferenceHandle).GetPooledBuffer());
-		BufferSRV = GraphBuilder.CreateSRV(Buffer, PF_R32_FLOAT);
-	}
-
 	void FVertexDeltaGraphDataProviderProxy::GatherDispatchData(FDispatchData const& InDispatchData)
 	{
 		const FSkeletalMeshRenderData& SkeletalMeshRenderData = SkeletalMeshObject->GetSkeletalMeshRenderData();
diff --git a/Engine/Plugins/Animation/MLDeformer/VertexDeltaModel/Source/VertexDeltaModel/Private/VertexDeltaGraphDebugDataInterface.cpp b/Engine/Plugins/Animation/MLDeformer/VertexDeltaModel/Source/VertexDeltaModel/Private/VertexDeltaGraphDebugDataInterface.cpp
--- a/Engine/Plugins/Animation/MLDeformer/VertexDeltaModel/Source/VertexDeltaModel/Private/VertexDeltaGraphDebugDataInterface.cpp
+++ b/Engine/Plugins/Animation/MLDeformer/VertexDeltaModel/Source/VertexDeltaModel/Private/VertexDeltaGraphDebugDataInterface.cpp
@@ -16,7 +16,7 @@
 #include "ShaderCompilerCore.h"
 #include "ShaderParameterMetadataBuilder.h"
 #include "SkeletalRenderPublic.h"
-#include "NeuralNetwork.h"
+
 
 #include UE_INLINE_GENERATED_CPP_BY_NAME(VertexDeltaGraphDebugDataInterface)
 
diff --git a/Engine/Plugins/Animation/MLDeformer/VertexDeltaModel/Source/VertexDeltaModel/Private/VertexDeltaModel.cpp b/Engine/Plugins/Animation/MLDeformer/VertexDeltaModel/Source/VertexDeltaModel/Private/VertexDeltaModel.cpp
--- a/Engine/Plugins/Animation/MLDeformer/VertexDeltaModel/Source/VertexDeltaModel/Private/VertexDeltaModel.cpp
+++ b/Engine/Plugins/Animation/MLDeformer/VertexDeltaModel/Source/VertexDeltaModel/Private/VertexDeltaModel.cpp
@@ -5,7 +5,9 @@
 #include "VertexDeltaModelInstance.h"
 #include "MLDeformerComponent.h"
 #include "MLDeformerAsset.h"
-#include "NeuralNetwork.h"
+#include "NNECore.h"
+#include "NNECoreModelData.h"
+#include "NNECoreRuntimeRDG.h"
 #include "Modules/ModuleManager.h"
 #include "Interfaces/IPluginManager.h"
 #include "Misc/Paths.h"
@@ -39,6 +41,11 @@ UVertexDeltaModel::UVertexDeltaModel(const FObjectInitializer& ObjectInitializer
 #endif
 }
 
+void UVertexDeltaModel::SetNNEModelData(TObjectPtr<UNNEModelData> ModelData)
+{
+	NNEModel = ModelData; 
+}
+
 UMLDeformerModelInstance* UVertexDeltaModel::CreateModelInstance(UMLDeformerComponent* Component)
 {
 	return NewObject<UVertexDeltaModelInstance>(Component);
@@ -52,26 +59,6 @@ FString UVertexDeltaModel::GetDefaultDeformerGraphAssetPath() const
 void UVertexDeltaModel::PostLoad()
 {
 	Super::PostLoad();
-
-	if (NNINetwork)
-	{
-		NNINetwork->SetDeviceType(ENeuralDeviceType::GPU, ENeuralDeviceType::CPU, ENeuralDeviceType::GPU);
-		if (NNINetwork->GetDeviceType() != ENeuralDeviceType::GPU || NNINetwork->GetOutputDeviceType() != ENeuralDeviceType::GPU || NNINetwork->GetInputDeviceType() != ENeuralDeviceType::CPU)
-		{
-			UE_LOG(LogVertexDeltaModel, Error, TEXT("Neural net in MLD Vertex Delta Model '%s' cannot run on the GPU, it will not be active."), *GetDeformerAsset()->GetName());
-		}
-	}
-}
-
-void UVertexDeltaModel::SetNNINetwork(UNeuralNetwork* InNeuralNetwork)
-{
-	GetNeuralNetworkModifyDelegate().Broadcast();
-	NNINetwork = InNeuralNetwork;
-}
-
-UNeuralNetwork* UVertexDeltaModel::GetNNINetwork() const
-{
-	return NNINetwork.Get();
 }
 
 #if WITH_EDITOR
@@ -79,17 +66,12 @@ UNeuralNetwork* UVertexDeltaModel::GetNNINetwork() const
 	{
 		Super::UpdateMemoryUsage();
 
-		// Check if the neural network is on the GPU, if so, count the memory to the GPU and remove it from Main memory, as we added it to 
-		// that already when we did the Model->GetResourceSizeBytes.
-		if (NNINetwork)
+		if (NNEModel)
 		{
-			if (NNINetwork->GetDeviceType() == ENeuralDeviceType::GPU)
-			{
-				const uint64 NeuralNetSize = static_cast<uint64>(NNINetwork->GetResourceSizeBytes(EResourceSizeMode::Type::EstimatedTotal));
+				const uint64 NeuralNetSize = 0;
 				GPUMemUsageInBytes += NeuralNetSize;
 				MemUsageInBytes -= NeuralNetSize;
 				CookedMemUsageInBytes -= NeuralNetSize;
-			}
 		}
 	}
 #endif
diff --git a/Engine/Plugins/Animation/MLDeformer/VertexDeltaModel/Source/VertexDeltaModel/Private/VertexDeltaModelInstance.cpp b/Engine/Plugins/Animation/MLDeformer/VertexDeltaModel/Source/VertexDeltaModel/Private/VertexDeltaModelInstance.cpp
--- a/Engine/Plugins/Animation/MLDeformer/VertexDeltaModel/Source/VertexDeltaModel/Private/VertexDeltaModelInstance.cpp
+++ b/Engine/Plugins/Animation/MLDeformer/VertexDeltaModel/Source/VertexDeltaModel/Private/VertexDeltaModelInstance.cpp
@@ -3,7 +3,6 @@
 #include "VertexDeltaModelInstance.h"
 #include "VertexDeltaModel.h"
 #include "MLDeformerAsset.h"
-#include "NeuralNetwork.h"
 #include "RenderGraphBuilder.h"
 #include "Components/SkeletalMeshComponent.h"
 
@@ -13,10 +12,6 @@ UVertexDeltaModel* UVertexDeltaModelInstance::GetVertexDeltaModel() const
 	return Cast<UVertexDeltaModel>(Model);
 }
 
-UNeuralNetwork* UVertexDeltaModelInstance::GetNNINetwork() const
-{
-	return GetVertexDeltaModel()->GetNNINetwork();
-}
 
 void UVertexDeltaModelInstance::Release()
 {
@@ -25,24 +20,25 @@ void UVertexDeltaModelInstance::Release()
 	// Destroy the neural network instance.
 	if (GetVertexDeltaModel())
 	{
-		UNeuralNetwork* NNINetwork = GetNNINetwork();
-		if (NNINetwork && NeuralNetworkInferenceHandle != -1)
-		{
-			NNINetwork->DestroyInferenceContext(NeuralNetworkInferenceHandle);
-			NeuralNetworkInferenceHandle = -1;
-		}
+
 	}
 }
 
+UE::NNECore::IModelRDG* UVertexDeltaModelInstance::GetNNEModelRDG() const
+{
+	return ModelRDG.Get();
+}
+
 FString UVertexDeltaModelInstance::CheckCompatibility(USkeletalMeshComponent* InSkelMeshComponent, bool bLogIssues)
 {
 	FString ErrorString = Super::CheckCompatibility(InSkelMeshComponent, bLogIssues);
 
 	// Verify the number of inputs versus the expected number of inputs.	
-	const UNeuralNetwork* NNINetwork = GetNNINetwork();
-	if (NNINetwork && NNINetwork->IsLoaded() && Model->GetDeformerAsset())
+	const UE::NNECore::IModelRDG* RDGModel = GetNNEModelRDG();
+	if (RDGModel && Model->GetDeformerAsset())
 	{
-		const int64 NumNeuralNetInputs = NNINetwork->GetInputTensor().Num();
+		TConstArrayView<UE::NNECore::FTensorDesc> GPUTensorDesc = RDGModel->GetInputTensorDescs();
+		const int64 NumNeuralNetInputs = GPUTensorDesc[0].GetShape().Rank();
 		const int64 NumDeformerAssetInputs = static_cast<int64>(Model->GetInputInfo()->CalcNumNeuralNetInputs(Model->GetNumFloatsPerBone(), Model->GetNumFloatsPerCurve()));
 		if (NumNeuralNetInputs != NumDeformerAssetInputs)
 		{
@@ -60,51 +56,157 @@ FString UVertexDeltaModelInstance::CheckCompatibility(USkeletalMeshComponent* In
 
 bool UVertexDeltaModelInstance::IsValidForDataProvider() const
 {
-	UNeuralNetwork* NNINetwork = GetNNINetwork();
-	if (NNINetwork == nullptr || !NNINetwork->IsLoaded())
-	{
-		return false;
-	}
-
-	// We expect to run on the GPU when using a data provider for the deformer graph system (Optimus).
-	// Make sure we're actually running the network on the GPU.
-	// Inputs are expected to come from the CPU though.
-	if (NNINetwork->GetDeviceType() != ENeuralDeviceType::GPU || NNINetwork->GetOutputDeviceType() != ENeuralDeviceType::GPU || NNINetwork->GetInputDeviceType() != ENeuralDeviceType::CPU)
-	{
-		return false;
-	}
-
-	return (Model->GetVertexMapBuffer().ShaderResourceViewRHI != nullptr) && (GetNeuralNetworkInferenceHandle() != -1);
+	return ModelRDG.IsValid(); 
 }
 
 void UVertexDeltaModelInstance::Execute(float ModelWeight)
 {
 	TRACE_CPUPROFILER_EVENT_SCOPE(UMLDeformerModelInstance::Execute)
 
-	UNeuralNetwork* NNINetwork = GetNNINetwork();
-	if (NNINetwork == nullptr)
-	{
-		return;
-	}
-
-	// Even if the model needs the GPU it is possible that the hardware does not support GPU evaluation
-	if (NNINetwork->GetDeviceType() == ENeuralDeviceType::GPU)
+	if (ModelRDG)
 	{
-		// NOTE: Inputs still come from the CPU.
-		check(NNINetwork->GetDeviceType() == ENeuralDeviceType::GPU && NNINetwork->GetInputDeviceType() == ENeuralDeviceType::CPU && NNINetwork->GetOutputDeviceType() == ENeuralDeviceType::GPU);
 		ENQUEUE_RENDER_COMMAND(RunNeuralNetwork)
 			(
-				[NNINetwork, Handle = NeuralNetworkInferenceHandle](FRHICommandListImmediate& RHICmdList)
+				[this](FRHICommandListImmediate& RHICmdList)
 				{
+					ERHIPipeline Pipeline = RHICmdList.GetPipeline();
+
+					if (Pipeline == ERHIPipeline::None)
+					{
+						RHICmdList.SwitchPipeline(ERHIPipeline::Graphics);
+					}
 					// Output deltas will be available on GPU for DeformerGraph via UMLDeformerDataProvider.
 					FRDGBuilder GraphBuilder(RHICmdList);
-					NNINetwork->Run(GraphBuilder, Handle);
+
+					// Build Input Bindings
+					TArray<UE::NNECore::FTensorBindingRDG> InputBindingsRDG;
+					UE::NNECore::FTensorBindingRDG& BindingRDG = InputBindingsRDG.Emplace_GetRef();
+					BindingRDG.Buffer = RDGInputBuffer;
+					GraphBuilder.QueueBufferUpload(BindingRDG.Buffer, NNEInputTensorBuffer.GetData(), NNEInputTensorBuffer.Num() * sizeof(float), ERDGInitialDataFlags::NoCopy);
+	
+					// Build Output Bindings
+					TArray<UE::NNECore::FTensorBindingRDG> OutputBindingsRDG;
+					UE::NNECore::FTensorBindingRDG& OutputBindingRDG = OutputBindingsRDG.Emplace_GetRef();
+					OutputBindingRDG.Buffer = RDGVertexDeltaBuffer;
+
+					if (!ModelRDG.IsValid())
+					{
+						return;
+					}
+					int32 Res = ModelRDG->EnqueueRDG(GraphBuilder, InputBindingsRDG, OutputBindingsRDG);
+	
 					GraphBuilder.Execute();
 				}
 		);
 	}
 }
 
+bool UVertexDeltaModelInstance::GetRDGVertexBufferDesc(TConstArrayView<UE::NNECore::FTensorDesc>& InOutputTensorDescs, FRDGBufferDesc& OutBufferDesc)
+{
+	if (InOutputTensorDescs.Num() > 0)
+	{
+		FRDGBufferDesc BufferDesc;
+		const uint32 ElemByteSize = InOutputTensorDescs[0].GetElemByteSize();
+		const UE::NNECore::FSymbolicTensorShape& SymShape = InOutputTensorDescs[0].GetShape();
+		for (int32 i = 1; i < InOutputTensorDescs.Num(); i++)
+		{
+			if (InOutputTensorDescs[i].GetElemByteSize() != ElemByteSize || SymShape != InOutputTensorDescs[i].GetShape())
+			{
+				return false;
+			}
+		}
+		// Create a single flat output buffer
+		const UE::NNECore::FTensorShape OutputShape = UE::NNECore::FTensorShape::MakeFromSymbolic(SymShape);
+		BufferDesc.BytesPerElement = ElemByteSize;
+		BufferDesc.NumElements = OutputShape.Volume() * InOutputTensorDescs.Num();
+		BufferDesc.Usage = EBufferUsageFlags::UnorderedAccess | EBufferUsageFlags::ShaderResource | EBufferUsageFlags::VertexBuffer;
+		return true; 
+	}
+	return false;
+}
+
+void UVertexDeltaModelInstance::CreateRDGBuffers(TConstArrayView<UE::NNECore::FTensorDesc>& OutputTensorDescs)
+{
+	// Make sure that we wait for Graph builder to finish
+	std::atomic<bool> bIsGraphBuilderDone(false);
+
+	ENQUEUE_RENDER_COMMAND(VertexDeltaModelInstance_CreateOuputRDGBuffer)(
+		[this, &OutputTensorDescs, &bIsGraphBuilderDone](FRHICommandListImmediate& RHICmdList)
+		{
+			FRDGBuilder Builder(RHICmdList);
+			FRDGBufferDesc VertexBufferDesc;
+			if (GetRDGVertexBufferDesc(OutputTensorDescs, VertexBufferDesc))
+			{
+				RDGVertexDeltaBuffer = Builder.CreateBuffer(VertexBufferDesc, TEXT("UVertexDeltaModelInstance_OutputBuffer"));
+			}
+
+			FRDGBufferDesc InputDesc = FRDGBufferDesc::CreateBufferDesc(sizeof(float), NNEInputTensorBuffer.Num());
+			InputDesc.Usage = EBufferUsageFlags(InputDesc.Usage | BUF_SourceCopy);
+
+			RDGInputBuffer = Builder.CreateBuffer(InputDesc, TEXT("UVertexDeltaModelInstance_InputBuffer"), ERDGBufferFlags::None);
+
+			Builder.Execute();
+			bIsGraphBuilderDone = true;
+		}
+	);
+
+	// Wait until the graph builder is done
+	while (!bIsGraphBuilderDone)
+	{
+		FPlatformProcess::Sleep(0.1e-3);
+	}
+
+}
+
+void UVertexDeltaModelInstance::PostMLDeformerComponentInit()
+{
+	if (!bNNECreationAttempted)
+	{
+		bNNECreationAttempted = true;
+		CreateNNEModel();
+	}
+}
+
+void UVertexDeltaModelInstance::CreateNNEModel()
+{
+	if (!ModelRDG.IsValid())
+	{
+		UVertexDeltaModel* VertexDeltaModel = Cast<UVertexDeltaModel>(Model);
+		if (VertexDeltaModel)
+		{
+			TWeakInterfacePtr<INNERuntime> Runtime = UE::NNECore::GetRuntime<INNERuntime>(VertexDeltaModel->GetNNERuntimeName());
+			TWeakInterfacePtr<INNERuntimeRDG> RuntimeRDG = UE::NNECore::GetRuntime<INNERuntimeRDG>(VertexDeltaModel->GetNNERuntimeName());
+
+			if (!Runtime.IsValid())
+			{
+				UE_LOG(LogNNE, Error, TEXT("Can't get %s runtime."), *VertexDeltaModel->GetNNERuntimeName());
+				return;
+			}
+
+			TObjectPtr<UNNEModelData> ModelData = VertexDeltaModel->NNEModel;
+			if (ModelData)
+			{
+				if (RuntimeRDG.IsValid())
+				{
+					// allocate tensor inputs and outputs
+					ModelRDG = RuntimeRDG->CreateModelRDG(ModelData);
+
+					// setup inputs
+					TConstArrayView<UE::NNECore::FTensorDesc> InputTensorDescs = ModelRDG->GetInputTensorDescs();
+					UE::NNECore::FTensorShape InputTensorShape = UE::NNECore::FTensorShape::MakeFromSymbolic(InputTensorDescs[0].GetShape());
+					ModelRDG->SetInputTensorShapes({ InputTensorShape });
+					check(InputTensorDescs[0].GetElemByteSize() == sizeof(float));
+					NNEInputTensorBuffer.SetNumUninitialized(InputTensorShape.Volume());
+
+					// setup outputs
+					TConstArrayView<UE::NNECore::FTensorDesc> OutputTensorDescs = ModelRDG->GetOutputTensorDescs();
+					CreateRDGBuffers(OutputTensorDescs);
+				}
+			}
+		}
+	}
+}
+
 bool UVertexDeltaModelInstance::SetupInputs()
 {
 	TRACE_CPUPROFILER_EVENT_SCOPE(UMLDeformerModelInstance::SetupInputs)
@@ -119,35 +221,28 @@ bool UVertexDeltaModelInstance::SetupInputs()
 	}
 
 	// Get the network and make sure it's loaded.
-	UNeuralNetwork* NNINetwork = GetNNINetwork();
-	if (NNINetwork == nullptr || !NNINetwork->IsLoaded())
-	{
-		return false;
-	}
 
-	// Allocate an inference context if none has already been allocated.
-	if (NeuralNetworkInferenceHandle == -1)
+	if (!ModelRDG)
 	{
-		NeuralNetworkInferenceHandle = NNINetwork->CreateInferenceContext();
-		if (NeuralNetworkInferenceHandle == -1)
-		{
-			return false;
-		}
+		return false;
 	}
 
-	// If the neural network expects a different number of inputs, do nothing.
-	const int64 NumNeuralNetInputs = NNINetwork->GetInputTensorForContext(NeuralNetworkInferenceHandle).Num();
+	TConstArrayView<UE::NNECore::FTensorDesc> InputTensorDescs = ModelRDG->GetInputTensorDescs();
+	const UE::NNECore::FTensorShape InputTensorShape = UE::NNECore::FTensorShape::MakeFromSymbolic(InputTensorDescs[0].GetShape());
+	const int64 NumNeuralNetInputs = InputTensorShape.Volume();
 	const int64 NumDeformerAssetInputs = Model->GetInputInfo()->CalcNumNeuralNetInputs(Model->GetNumFloatsPerBone(), Model->GetNumFloatsPerCurve());
 	if (NumNeuralNetInputs != NumDeformerAssetInputs)
 	{
 		return false;
 	}
 
+	if (NNEInputTensorBuffer.Num() != NumNeuralNetInputs)
+	{
+		return false; 
+	}
 	// Update and write the input values directly into the input tensor.
-	float* InputDataPointer = static_cast<float*>(NNINetwork->GetInputDataPointerMutableForContext(NeuralNetworkInferenceHandle));
-	const int64 NumFloatsWritten = SetNeuralNetworkInputValues(InputDataPointer, NumNeuralNetInputs);
+	const int64 NumFloatsWritten = SetNeuralNetworkInputValues(NNEInputTensorBuffer.GetData(), NNEInputTensorBuffer.Num());
 	check(NumFloatsWritten == NumNeuralNetInputs);
-
 	return true;
 }
 
diff --git a/Engine/Plugins/Animation/MLDeformer/VertexDeltaModel/Source/VertexDeltaModel/Public/VertexDeltaGraphDataInterface.h b/Engine/Plugins/Animation/MLDeformer/VertexDeltaModel/Source/VertexDeltaModel/Public/VertexDeltaGraphDataInterface.h
--- a/Engine/Plugins/Animation/MLDeformer/VertexDeltaModel/Source/VertexDeltaModel/Public/VertexDeltaGraphDataInterface.h
+++ b/Engine/Plugins/Animation/MLDeformer/VertexDeltaModel/Source/VertexDeltaModel/Public/VertexDeltaGraphDataInterface.h
@@ -8,6 +8,9 @@
 #include "Rendering/SkeletalMeshRenderData.h"
 #include "Rendering/SkeletalMeshLODRenderData.h"
 #include "SkeletalRenderPublic.h"
+#include "NNECore.h"
+#include "NNECoreModelData.h"
+#include "NNECoreRuntimeRDG.h"
 #include "VertexDeltaGraphDataInterface.generated.h"
 
 class FRDGBuffer;
@@ -15,9 +18,9 @@ class FRDGBufferSRV;
 class FRHIShaderResourceView;
 class FSkeletalMeshObject;
 class UMLDeformerComponent;
-class UNeuralNetwork;
 class USkeletalMeshComponent;
 class UMLDeformerModel;
+class UVertexDeltaModelInstance;
 
 /** Compute Framework Data Interface for MLDefomer data. */
 UCLASS(Category = ComputeFramework)
@@ -76,13 +79,13 @@ namespace UE::VertexDeltaModel
 
 		// FComputeDataProviderRenderProxy overrides.
 		bool IsValid(FValidationData const& InValidationData) const override;
-		void AllocateResources(FRDGBuilder& GraphBuilder) override;
 		void GatherDispatchData(FDispatchData const& InDispatchData) override;
 		// ~END FComputeDataProviderRenderProxy overrides.
 
 	protected:
 		FSkeletalMeshObject* SkeletalMeshObject = nullptr;
-		TObjectPtr<UNeuralNetwork> NeuralNetwork = nullptr;
+		UE::NNECore::IModelRDG* NNEModelRDG = nullptr;
+		UVertexDeltaModelInstance* VertexDeltaModelInstance = nullptr;
 		FRHIShaderResourceView* VertexMapBufferSRV = nullptr;
 		FRDGBuffer* Buffer = nullptr;
 		FRDGBufferSRV* BufferSRV = nullptr;
diff --git a/Engine/Plugins/Animation/MLDeformer/VertexDeltaModel/Source/VertexDeltaModel/Public/VertexDeltaModel.h b/Engine/Plugins/Animation/MLDeformer/VertexDeltaModel/Source/VertexDeltaModel/Public/VertexDeltaModel.h
--- a/Engine/Plugins/Animation/MLDeformer/VertexDeltaModel/Source/VertexDeltaModel/Public/VertexDeltaModel.h
+++ b/Engine/Plugins/Animation/MLDeformer/VertexDeltaModel/Source/VertexDeltaModel/Public/VertexDeltaModel.h
@@ -5,12 +5,13 @@
 #include "CoreTypes.h"
 #include "MLDeformerGeomCacheModel.h"
 #include "VertexDeltaModelVizSettings.h"
+#include "NNECoreRuntimeRDG.h"
 #include "VertexDeltaModel.generated.h"
 
 VERTEXDELTAMODEL_API DECLARE_LOG_CATEGORY_EXTERN(LogVertexDeltaModel, Log, All);
 
 class UNeuralNetwork;
-
+class UNNEModelData;
 /** 
  * The vertex delta model, which uses a GPU based neural network.
  * This model acts more as an example of how to implement a model that only uses the GPU.
@@ -41,9 +42,6 @@ public:
 #endif
 	// ~END UMLDeformerModel overrides.
 
-	UNeuralNetwork* GetNNINetwork() const;
-	void SetNNINetwork(UNeuralNetwork* InNeuralNetwork);
-
 #if WITH_EDITORONLY_DATA
 	int32 GetNumHiddenLayers() const						{ return NumHiddenLayers; }
 	int32 GetNumNeuronsPerLayer() const						{ return NumNeuronsPerLayer; }
@@ -51,12 +49,13 @@ public:
 	int32 GetBatchSize() const								{ return BatchSize; }
 	float GetLearningRate() const							{ return LearningRate; }
 #endif
-
-
+	void SetNNEModelData(TObjectPtr<UNNEModelData> ModelData);
+	const FString GetNNERuntimeName() const { return TEXT("NNERuntimeRDGDml"); }
 public:
-	/** The NNI neural network. */
+
+	/** The NNE neural network model. */
 	UPROPERTY()
-	TObjectPtr<UNeuralNetwork> NNINetwork;
+	TObjectPtr<UNNEModelData> NNEModel; 
 
 #if WITH_EDITORONLY_DATA
 	/** The number of hidden layers that the neural network model will have.\nHigher numbers will slow down performance but can deal with more complex deformations. */
@@ -79,4 +78,5 @@ public:
 	UPROPERTY(EditAnywhere, BlueprintReadWrite, Category = "Training Settings", meta = (ClampMin = "0.000001", ClampMax = "1.0"))
 	float LearningRate = 0.001f;
 #endif // WITH_EDITORONLY_DATA
+
 };
diff --git a/Engine/Plugins/Animation/MLDeformer/VertexDeltaModel/Source/VertexDeltaModel/Public/VertexDeltaModelInstance.h b/Engine/Plugins/Animation/MLDeformer/VertexDeltaModel/Source/VertexDeltaModel/Public/VertexDeltaModelInstance.h
--- a/Engine/Plugins/Animation/MLDeformer/VertexDeltaModel/Source/VertexDeltaModel/Public/VertexDeltaModelInstance.h
+++ b/Engine/Plugins/Animation/MLDeformer/VertexDeltaModel/Source/VertexDeltaModel/Public/VertexDeltaModelInstance.h
@@ -3,6 +3,10 @@
 #pragma once
 
 #include "MLDeformerModelInstance.h"
+#include "NNECore.h"
+#include "NNECoreRuntime.h"
+#include "NNECoreModelData.h"
+#include "NNECoreRuntimeRDG.h"
 #include "VertexDeltaModelInstance.generated.h"
 
 class UVertexDeltaModel;
@@ -24,5 +28,43 @@ public:
 	// ~END UMLDeformerModelInstance overrides.
 
 	UVertexDeltaModel* GetVertexDeltaModel() const;
-	UNeuralNetwork* GetNNINetwork() const;
+	virtual void PostMLDeformerComponentInit() override;
+
+	/**
+	 * Get the RDG Neural Network Model
+	 * @return The IModelRDG if successfully created, or nullptr
+	 */
+	UE::NNECore::IModelRDG* GetNNEModelRDG() const;
+
+	/**
+	 * Get the output vertex delta buffer
+	 * @return The FRDGBuffer for vertex deltas
+	 */
+	FRDGBuffer* GetOutputRDGBuffer() const;
+
+	/**
+	 * Get the render graph buffer description required for the output of a neural network. Return false if a flat buffer is not appropriate
+	 * @return True if the OutputTensorDescs can be rppresented by a flat float  FRDGBuffer
+	 */
+	bool GetRDGVertexBufferDesc(TConstArrayView<UE::NNECore::FTensorDesc>& InOutputTensorDescs, FRDGBufferDesc& OutBufferDesc);
+
+private:
+
+	void CreateRDGBuffers(TConstArrayView<UE::NNECore::FTensorDesc>& OutputTensorDescs);
+	void CreateNNEModel();
+
+	// Input Buffer for Joint Matrices / Curve Floats
+	FRDGBuffer* RDGInputBuffer = nullptr;
+
+	// Output Buffer for Vertex Deltas
+	FRDGBuffer* RDGVertexDeltaBuffer = nullptr;
+
+	// The NNE RDG Model 
+	TUniquePtr<UE::NNECore::IModelRDG> ModelRDG;
+	
+	// The CPU Input Tensor Buffer
+	TArray<float> NNEInputTensorBuffer;
+
+	// Only attempt to create NNE Model once
+	bool bNNECreationAttempted = false;
 };
diff --git a/Engine/Plugins/Animation/MLDeformer/VertexDeltaModel/Source/VertexDeltaModel/VertexDeltaModel.Build.cs b/Engine/Plugins/Animation/MLDeformer/VertexDeltaModel/Source/VertexDeltaModel/VertexDeltaModel.Build.cs
--- a/Engine/Plugins/Animation/MLDeformer/VertexDeltaModel/Source/VertexDeltaModel/VertexDeltaModel.Build.cs
+++ b/Engine/Plugins/Animation/MLDeformer/VertexDeltaModel/Source/VertexDeltaModel/VertexDeltaModel.Build.cs
@@ -38,7 +38,9 @@ public class VertexDeltaModel : ModuleRules
 				"CoreUObject",
 				"Engine",
 				"GeometryCache",
-				"NeuralNetworkInference",
+				"NNECore",
+				"NNERuntimeRDG",
+				"NNEUtils",
 				"OptimusCore",
 				"Projects",
 				"RenderCore",
diff --git a/Engine/Plugins/Animation/MLDeformer/VertexDeltaModel/Source/VertexDeltaModelEditor/Private/VertexDeltaEditorModel.cpp b/Engine/Plugins/Animation/MLDeformer/VertexDeltaModel/Source/VertexDeltaModelEditor/Private/VertexDeltaEditorModel.cpp
--- a/Engine/Plugins/Animation/MLDeformer/VertexDeltaModel/Source/VertexDeltaModelEditor/Private/VertexDeltaEditorModel.cpp
+++ b/Engine/Plugins/Animation/MLDeformer/VertexDeltaModel/Source/VertexDeltaModelEditor/Private/VertexDeltaEditorModel.cpp
@@ -5,7 +5,10 @@
 #include "VertexDeltaModelVizSettings.h"
 #include "VertexDeltaTrainingModel.h"
 #include "MLDeformerAsset.h"
-#include "NeuralNetwork.h"
+#include "NNECore.h"
+#include "NNECoreModelData.h"
+#include "Misc/Paths.h"
+#include "Misc/FileHelper.h"
 
 #define LOCTEXT_NAMESPACE "VertexDeltaEditorModel"
 
@@ -28,22 +31,23 @@ namespace UE::VertexDeltaModel
 		return TrainModel<UVertexDeltaTrainingModel>(this);
 	}
 
-	UNeuralNetwork* FVertexDeltaEditorModel::LoadNeuralNetworkFromOnnx(const FString& Filename) const
+	TObjectPtr<UNNEModelData>  FVertexDeltaEditorModel::LoadNeuralNetworkFromOnnx(const FString& Filename) const
 	{
 		const FString OnnxFile = FPaths::ConvertRelativePathToFull(Filename);
 		if (FPaths::FileExists(OnnxFile))
 		{
 			UE_LOG(LogVertexDeltaModel, Display, TEXT("Loading Onnx file '%s'..."), *OnnxFile);
-			UNeuralNetwork* Result = NewObject<UNeuralNetwork>(Model, UNeuralNetwork::StaticClass());		
-			if (Result->Load(OnnxFile))
+			TArray<uint8> RawModelData;
+			const bool bIsModelInMem = FFileHelper::LoadFileToArray(RawModelData, *OnnxFile);
+	
+			if (bIsModelInMem)
 			{
-				Result->SetDeviceType(ENeuralDeviceType::GPU, ENeuralDeviceType::CPU, ENeuralDeviceType::GPU);	
-				if (Result->GetDeviceType() != ENeuralDeviceType::GPU || Result->GetOutputDeviceType() != ENeuralDeviceType::GPU || Result->GetInputDeviceType() != ENeuralDeviceType::CPU)
-				{
-					UE_LOG(LogVertexDeltaModel, Error, TEXT("Neural net in ML Deformer '%s' cannot run on the GPU, it will not be active."), *Model->GetDeformerAsset()->GetName());
-				}
+				TObjectPtr<UNNEModelData> ModelData = NewObject<UNNEModelData>();
+
+				ModelData->Init(FString("onnx"), RawModelData);
+				// TODO - add some real verification
 				UE_LOG(LogVertexDeltaModel, Display, TEXT("Successfully loaded Onnx file '%s'..."), *OnnxFile);
-				return Result;
+				return ModelData;
 			}
 			else
 			{
@@ -55,16 +59,16 @@ namespace UE::VertexDeltaModel
 			UE_LOG(LogVertexDeltaModel, Error, TEXT("Onnx file '%s' does not exist!"), *OnnxFile);
 		}
 
-		return nullptr;
+		return TObjectPtr<UNNEModelData>();
 	}
 
 	bool FVertexDeltaEditorModel::LoadTrainedNetwork() const
 	{
 		const FString OnnxFile = GetTrainedNetworkOnnxFile();
-		UNeuralNetwork* Network = LoadNeuralNetworkFromOnnx(OnnxFile);
-		if (Network)
+		TObjectPtr<UNNEModelData> ModelData = LoadNeuralNetworkFromOnnx(OnnxFile);
+		if (ModelData)
 		{
-			GetVertexDeltaModel()->SetNNINetwork(Network);
+			GetVertexDeltaModel()->SetNNEModelData(ModelData);
 			return true;
 		}
 
@@ -73,7 +77,7 @@ namespace UE::VertexDeltaModel
 
 	bool FVertexDeltaEditorModel::IsTrained() const
 	{
-		return (GetVertexDeltaModel()->GetNNINetwork() != nullptr);
+		return GetVertexDeltaModel()->NNEModel != nullptr;
 	}
 
 	FString FVertexDeltaEditorModel::GetHeatMapDeformerGraphPath() const
diff --git a/Engine/Plugins/Animation/MLDeformer/VertexDeltaModel/Source/VertexDeltaModelEditor/Public/VertexDeltaEditorModel.h b/Engine/Plugins/Animation/MLDeformer/VertexDeltaModel/Source/VertexDeltaModelEditor/Public/VertexDeltaEditorModel.h
--- a/Engine/Plugins/Animation/MLDeformer/VertexDeltaModel/Source/VertexDeltaModelEditor/Public/VertexDeltaEditorModel.h
+++ b/Engine/Plugins/Animation/MLDeformer/VertexDeltaModel/Source/VertexDeltaModelEditor/Public/VertexDeltaEditorModel.h
@@ -7,6 +7,7 @@
 
 class UNeuralNetwork;
 class UVertexDeltaModel;
+class UNNEModelData;
 
 namespace UE::VertexDeltaModel
 {
@@ -33,7 +34,7 @@ namespace UE::VertexDeltaModel
 		virtual FString GetHeatMapDeformerGraphPath() const override;
 		// ~END FMLDeformerEditorModel overrides.
 
-		UNeuralNetwork* LoadNeuralNetworkFromOnnx(const FString& Filename) const;
+		TObjectPtr<UNNEModelData> LoadNeuralNetworkFromOnnx(const FString& Filename) const;
 		UVertexDeltaModel* GetVertexDeltaModel() const;
 	};
 }	// namespace UE::VertexDeltaModel
diff --git a/Engine/Plugins/Animation/MLDeformer/VertexDeltaModel/Source/VertexDeltaModelEditor/VertexDeltaModelEditor.Build.cs b/Engine/Plugins/Animation/MLDeformer/VertexDeltaModel/Source/VertexDeltaModelEditor/VertexDeltaModelEditor.Build.cs
--- a/Engine/Plugins/Animation/MLDeformer/VertexDeltaModel/Source/VertexDeltaModelEditor/VertexDeltaModelEditor.Build.cs
+++ b/Engine/Plugins/Animation/MLDeformer/VertexDeltaModel/Source/VertexDeltaModelEditor/VertexDeltaModelEditor.Build.cs
@@ -31,7 +31,9 @@ namespace UnrealBuildTool.Rules
 					"MLDeformerFrameworkEditor",
 					"VertexDeltaModel",
 					"PropertyEditor",
-					"NeuralNetworkInference",
+					"NNECore",
+					"NNERuntimeRDG",
+					"NNEUtils",
 					"ToolWidgets",
 					"ComputeFramework"
 				}
diff --git a/Engine/Plugins/Animation/MLDeformer/VertexDeltaModel/VertexDeltaModel.uplugin b/Engine/Plugins/Animation/MLDeformer/VertexDeltaModel/VertexDeltaModel.uplugin
--- a/Engine/Plugins/Animation/MLDeformer/VertexDeltaModel/VertexDeltaModel.uplugin
+++ b/Engine/Plugins/Animation/MLDeformer/VertexDeltaModel/VertexDeltaModel.uplugin
@@ -14,16 +14,23 @@
 	"IsBetaVersion": true,
 	"IsExperimentalVersion": false,
 	"Installed": false,
+	"SupportedTargetPlatforms": [
+		"Win64",
+		"Linux",
+		"Mac"
+	],
 	"Modules": [
 		{
 			"Name": "VertexDeltaModel",
 			"Type": "Runtime",
-			"LoadingPhase": "Default"
+			"LoadingPhase": "Default",
+			"PlatformAllowList": [ "Win64", "Linux", "Mac" ]
 		},
 		{
 			"Name": "VertexDeltaModelEditor",
 			"Type": "Editor",
-			"LoadingPhase": "Default"
+			"LoadingPhase": "Default",
+			"PlatformAllowList": [ "Win64", "Linux", "Mac" ]
 		}
 	],
 	"Plugins": [
@@ -32,7 +39,7 @@
 			"Enabled": true
 		},
 		{
-			"Name": "NeuralNetworkInference",
+			"Name": "NNE",
 			"Enabled": true
 		},
 		{
