[MLDeformer] Show a warning in the Vertex Delta Model that it shouldn't be used in production.

#trivial

diff --git a/Engine/Plugins/Animation/MLDeformer/VertexDeltaModel/Source/VertexDeltaModelEditor/Private/VertexDeltaModelDetails.cpp b/Engine/Plugins/Animation/MLDeformer/VertexDeltaModel/Source/VertexDeltaModelEditor/Private/VertexDeltaModelDetails.cpp
--- a/Engine/Plugins/Animation/MLDeformer/VertexDeltaModel/Source/VertexDeltaModelEditor/Private/VertexDeltaModelDetails.cpp
+++ b/Engine/Plugins/Animation/MLDeformer/VertexDeltaModel/Source/VertexDeltaModelEditor/Private/VertexDeltaModelDetails.cpp
@@ -3,6 +3,9 @@
 #include "VertexDeltaModelDetails.h"
 #include "VertexDeltaModel.h"
 #include "DetailCategoryBuilder.h"
+#include "DetailWidgetRow.h"
+#include "Widgets/Layout/SBox.h"
+#include "SWarningOrErrorBox.h"
 
 #define LOCTEXT_NAMESPACE "VertexDeltaModelDetails"
 
@@ -26,6 +29,18 @@ namespace UE::VertexDeltaModel
 		TrainingSettingsCategoryBuilder->AddProperty(GET_MEMBER_NAME_CHECKED(UVertexDeltaModel, NumIterations));
 		TrainingSettingsCategoryBuilder->AddProperty(GET_MEMBER_NAME_CHECKED(UVertexDeltaModel, BatchSize));
 		TrainingSettingsCategoryBuilder->AddProperty(GET_MEMBER_NAME_CHECKED(UVertexDeltaModel, LearningRate));
+
+		FDetailWidgetRow& ExampleWarningRow = TrainingSettingsCategoryBuilder->AddCustomRow(FText::FromString("VertexDeltaModelWarning"))
+			.WholeRowContent()
+			[
+				SNew(SBox)
+				.Padding(FMargin(0.0f, 4.0f))
+				[
+					SNew(SWarningOrErrorBox)
+					.MessageStyle(EMessageStyle::Warning)
+					.Message(LOCTEXT("VertexDeltaModelWarningText", "The Vertex Delta Model is a GPU based example model that shouldn't be used in production. Please use the Neural Morph Model with the mode set to 'Global' to achieve similar results in production."))
+				]
+			];
 	}
 }	// namespace UE::VertexDeltaModel
 
