MLDeformer: Fixed inference bugs with Vertex Delta Model

diff --git a/Engine/Plugins/Animation/MLDeformer/VertexDeltaModel/Source/VertexDeltaModel/Private/VertexDeltaGraphDataInterface.cpp b/Engine/Plugins/Animation/MLDeformer/VertexDeltaModel/Source/VertexDeltaModel/Private/VertexDeltaGraphDataInterface.cpp
--- a/Engine/Plugins/Animation/MLDeformer/VertexDeltaModel/Source/VertexDeltaModel/Private/VertexDeltaGraphDataInterface.cpp
+++ b/Engine/Plugins/Animation/MLDeformer/VertexDeltaModel/Source/VertexDeltaModel/Private/VertexDeltaGraphDataInterface.cpp
@@ -142,6 +142,16 @@ namespace UE::VertexDeltaModel
 		return true;
 	}
 
+
+	void FVertexDeltaGraphDataProviderProxy::AllocateResources(FRDGBuilder& GraphBuilder)
+	{
+		if (VertexDeltaModelInstance && VertexDeltaModelInstance->IsValidForDataProvider())
+		{
+			Buffer = VertexDeltaModelInstance->GetOutputRDGBuffer();
+			BufferSRV = GraphBuilder.CreateSRV(Buffer, PF_R32_FLOAT);
+		}
+	}
+
 	void FVertexDeltaGraphDataProviderProxy::GatherDispatchData(FDispatchData const& InDispatchData)
 	{
 		const FSkeletalMeshRenderData& SkeletalMeshRenderData = SkeletalMeshObject->GetSkeletalMeshRenderData();
diff --git a/Engine/Plugins/Animation/MLDeformer/VertexDeltaModel/Source/VertexDeltaModel/Private/VertexDeltaModelInstance.cpp b/Engine/Plugins/Animation/MLDeformer/VertexDeltaModel/Source/VertexDeltaModel/Private/VertexDeltaModelInstance.cpp
--- a/Engine/Plugins/Animation/MLDeformer/VertexDeltaModel/Source/VertexDeltaModel/Private/VertexDeltaModelInstance.cpp
+++ b/Engine/Plugins/Animation/MLDeformer/VertexDeltaModel/Source/VertexDeltaModel/Private/VertexDeltaModelInstance.cpp
@@ -29,6 +29,11 @@ UE::NNECore::IModelRDG* UVertexDeltaModelInstance::GetNNEModelRDG() const
 	return ModelRDG.Get();
 }
 
+FRDGBuffer* UVertexDeltaModelInstance::GetOutputRDGBuffer() const
+{
+	return RDGVertexDeltaBuffer; 
+}
+
 FString UVertexDeltaModelInstance::CheckCompatibility(USkeletalMeshComponent* InSkelMeshComponent, bool bLogIssues)
 {
 	FString ErrorString = Super::CheckCompatibility(InSkelMeshComponent, bLogIssues);
@@ -105,7 +110,6 @@ bool UVertexDeltaModelInstance::GetRDGVertexBufferDesc(TConstArrayView<UE::NNECo
 {
 	if (InOutputTensorDescs.Num() > 0)
 	{
-		FRDGBufferDesc BufferDesc;
 		const uint32 ElemByteSize = InOutputTensorDescs[0].GetElemByteSize();
 		const UE::NNECore::FSymbolicTensorShape& SymShape = InOutputTensorDescs[0].GetShape();
 		for (int32 i = 1; i < InOutputTensorDescs.Num(); i++)
@@ -117,9 +121,9 @@ bool UVertexDeltaModelInstance::GetRDGVertexBufferDesc(TConstArrayView<UE::NNECo
 		}
 		// Create a single flat output buffer
 		const UE::NNECore::FTensorShape OutputShape = UE::NNECore::FTensorShape::MakeFromSymbolic(SymShape);
-		BufferDesc.BytesPerElement = ElemByteSize;
-		BufferDesc.NumElements = OutputShape.Volume() * InOutputTensorDescs.Num();
-		BufferDesc.Usage = EBufferUsageFlags::UnorderedAccess | EBufferUsageFlags::ShaderResource | EBufferUsageFlags::VertexBuffer;
+		OutBufferDesc.BytesPerElement = ElemByteSize;
+		OutBufferDesc.NumElements = OutputShape.Volume() * InOutputTensorDescs.Num();
+		OutBufferDesc.Usage = EBufferUsageFlags::UnorderedAccess | EBufferUsageFlags::ShaderResource | EBufferUsageFlags::VertexBuffer;
 		return true; 
 	}
 	return false;
@@ -127,11 +131,8 @@ bool UVertexDeltaModelInstance::GetRDGVertexBufferDesc(TConstArrayView<UE::NNECo
 
 void UVertexDeltaModelInstance::CreateRDGBuffers(TConstArrayView<UE::NNECore::FTensorDesc>& OutputTensorDescs)
 {
-	// Make sure that we wait for Graph builder to finish
-	std::atomic<bool> bIsGraphBuilderDone(false);
-
 	ENQUEUE_RENDER_COMMAND(VertexDeltaModelInstance_CreateOuputRDGBuffer)(
-		[this, &OutputTensorDescs, &bIsGraphBuilderDone](FRHICommandListImmediate& RHICmdList)
+		[this, &OutputTensorDescs](FRHICommandListImmediate& RHICmdList)
 		{
 			FRDGBuilder Builder(RHICmdList);
 			FRDGBufferDesc VertexBufferDesc;
@@ -146,15 +147,12 @@ void UVertexDeltaModelInstance::CreateRDGBuffers(TConstArrayView<UE::NNECore::FT
 			RDGInputBuffer = Builder.CreateBuffer(InputDesc, TEXT("UVertexDeltaModelInstance_InputBuffer"), ERDGBufferFlags::None);
 
 			Builder.Execute();
-			bIsGraphBuilderDone = true;
 		}
 	);
 
-	// Wait until the graph builder is done
-	while (!bIsGraphBuilderDone)
-	{
-		FPlatformProcess::Sleep(0.1e-3);
-	}
+	FRenderCommandFence RenderFence;
+	RenderFence.BeginFence();
+	RenderFence.Wait();
 
 }
 
@@ -191,16 +189,24 @@ void UVertexDeltaModelInstance::CreateNNEModel()
 					// allocate tensor inputs and outputs
 					ModelRDG = RuntimeRDG->CreateModelRDG(ModelData);
 
-					// setup inputs
-					TConstArrayView<UE::NNECore::FTensorDesc> InputTensorDescs = ModelRDG->GetInputTensorDescs();
-					UE::NNECore::FTensorShape InputTensorShape = UE::NNECore::FTensorShape::MakeFromSymbolic(InputTensorDescs[0].GetShape());
-					ModelRDG->SetInputTensorShapes({ InputTensorShape });
-					check(InputTensorDescs[0].GetElemByteSize() == sizeof(float));
-					NNEInputTensorBuffer.SetNumUninitialized(InputTensorShape.Volume());
-
-					// setup outputs
-					TConstArrayView<UE::NNECore::FTensorDesc> OutputTensorDescs = ModelRDG->GetOutputTensorDescs();
-					CreateRDGBuffers(OutputTensorDescs);
+					if (ModelRDG)
+					{
+						// setup inputs
+						TConstArrayView<UE::NNECore::FTensorDesc> InputTensorDescs = ModelRDG->GetInputTensorDescs();
+						UE::NNECore::FTensorShape InputTensorShape = UE::NNECore::FTensorShape::MakeFromSymbolic(InputTensorDescs[0].GetShape());
+						ModelRDG->SetInputTensorShapes({ InputTensorShape });
+						check(InputTensorDescs[0].GetElemByteSize() == sizeof(float));
+						NNEInputTensorBuffer.SetNumUninitialized(InputTensorShape.Volume());
+
+						// setup outputs
+						TConstArrayView<UE::NNECore::FTensorDesc> OutputTensorDescs = ModelRDG->GetOutputTensorDescs();
+						CreateRDGBuffers(OutputTensorDescs);
+					}
+					else
+					{
+						UE_LOG(LogNNE, Error, TEXT("Failed to create NNE RDG Model for VertexDeltaModel."));
+						return;
+					}
 				}
 			}
 		}
diff --git a/Engine/Plugins/Animation/MLDeformer/VertexDeltaModel/Source/VertexDeltaModel/Public/VertexDeltaGraphDataInterface.h b/Engine/Plugins/Animation/MLDeformer/VertexDeltaModel/Source/VertexDeltaModel/Public/VertexDeltaGraphDataInterface.h
--- a/Engine/Plugins/Animation/MLDeformer/VertexDeltaModel/Source/VertexDeltaModel/Public/VertexDeltaGraphDataInterface.h
+++ b/Engine/Plugins/Animation/MLDeformer/VertexDeltaModel/Source/VertexDeltaModel/Public/VertexDeltaGraphDataInterface.h
@@ -80,6 +80,7 @@ namespace UE::VertexDeltaModel
 		// FComputeDataProviderRenderProxy overrides.
 		bool IsValid(FValidationData const& InValidationData) const override;
 		void GatherDispatchData(FDispatchData const& InDispatchData) override;
+		void AllocateResources(FRDGBuilder& GraphBuilder) override;
 		// ~END FComputeDataProviderRenderProxy overrides.
 
 	protected:
