Fix for render graph errors in vertex delta model

diff --git a/Engine/Plugins/Animation/MLDeformer/VertexDeltaModel/Source/VertexDeltaModel/Private/VertexDeltaGraphDataInterface.cpp b/Engine/Plugins/Animation/MLDeformer/VertexDeltaModel/Source/VertexDeltaModel/Private/VertexDeltaGraphDataInterface.cpp
--- a/Engine/Plugins/Animation/MLDeformer/VertexDeltaModel/Source/VertexDeltaModel/Private/VertexDeltaGraphDataInterface.cpp
+++ b/Engine/Plugins/Animation/MLDeformer/VertexDeltaModel/Source/VertexDeltaModel/Private/VertexDeltaGraphDataInterface.cpp
@@ -147,7 +147,7 @@ namespace UE::VertexDeltaModel
 	{
 		if (VertexDeltaModelInstance && VertexDeltaModelInstance->IsValidForDataProvider())
 		{
-			Buffer = VertexDeltaModelInstance->GetOutputRDGBuffer();
+			Buffer = GraphBuilder.RegisterExternalBuffer(VertexDeltaModelInstance->GetOutputRDGBuffer());
 			BufferSRV = GraphBuilder.CreateSRV(Buffer, PF_R32_FLOAT);
 		}
 	}
diff --git a/Engine/Plugins/Animation/MLDeformer/VertexDeltaModel/Source/VertexDeltaModel/Private/VertexDeltaModelInstance.cpp b/Engine/Plugins/Animation/MLDeformer/VertexDeltaModel/Source/VertexDeltaModel/Private/VertexDeltaModelInstance.cpp
--- a/Engine/Plugins/Animation/MLDeformer/VertexDeltaModel/Source/VertexDeltaModel/Private/VertexDeltaModelInstance.cpp
+++ b/Engine/Plugins/Animation/MLDeformer/VertexDeltaModel/Source/VertexDeltaModel/Private/VertexDeltaModelInstance.cpp
@@ -29,7 +29,7 @@ UE::NNECore::IModelRDG* UVertexDeltaModelInstance::GetNNEModelRDG() const
 	return ModelRDG.Get();
 }
 
-FRDGBuffer* UVertexDeltaModelInstance::GetOutputRDGBuffer() const
+TRefCountPtr<FRDGPooledBuffer> UVertexDeltaModelInstance::GetOutputRDGBuffer() const
 {
 	return RDGVertexDeltaBuffer; 
 }
@@ -86,13 +86,13 @@ void UVertexDeltaModelInstance::Execute(float ModelWeight)
 					// Build Input Bindings
 					TArray<UE::NNECore::FTensorBindingRDG> InputBindingsRDG;
 					UE::NNECore::FTensorBindingRDG& BindingRDG = InputBindingsRDG.Emplace_GetRef();
-					BindingRDG.Buffer = RDGInputBuffer;
+					BindingRDG.Buffer = GraphBuilder.RegisterExternalBuffer(RDGInputBuffer);
 					GraphBuilder.QueueBufferUpload(BindingRDG.Buffer, NNEInputTensorBuffer.GetData(), NNEInputTensorBuffer.Num() * sizeof(float), ERDGInitialDataFlags::NoCopy);
 	
 					// Build Output Bindings
 					TArray<UE::NNECore::FTensorBindingRDG> OutputBindingsRDG;
 					UE::NNECore::FTensorBindingRDG& OutputBindingRDG = OutputBindingsRDG.Emplace_GetRef();
-					OutputBindingRDG.Buffer = RDGVertexDeltaBuffer;
+					OutputBindingRDG.Buffer = GraphBuilder.RegisterExternalBuffer(RDGVertexDeltaBuffer);
 
 					if (!ModelRDG.IsValid())
 					{
@@ -138,13 +138,15 @@ void UVertexDeltaModelInstance::CreateRDGBuffers(TConstArrayView<UE::NNECore::FT
 			FRDGBufferDesc VertexBufferDesc;
 			if (GetRDGVertexBufferDesc(OutputTensorDescs, VertexBufferDesc))
 			{
-				RDGVertexDeltaBuffer = Builder.CreateBuffer(VertexBufferDesc, TEXT("UVertexDeltaModelInstance_OutputBuffer"));
+				FRDGBuffer* RDGBuffer = Builder.CreateBuffer(VertexBufferDesc, TEXT("UVertexDeltaModelInstance_OutputBuffer"));
+				RDGVertexDeltaBuffer = Builder.ConvertToExternalBuffer(RDGBuffer);
 			}
 
 			FRDGBufferDesc InputDesc = FRDGBufferDesc::CreateBufferDesc(sizeof(float), NNEInputTensorBuffer.Num());
 			InputDesc.Usage = EBufferUsageFlags(InputDesc.Usage | BUF_SourceCopy);
 
-			RDGInputBuffer = Builder.CreateBuffer(InputDesc, TEXT("UVertexDeltaModelInstance_InputBuffer"), ERDGBufferFlags::None);
+			FRDGBuffer* RDGBuffer = Builder.CreateBuffer(InputDesc, TEXT("UVertexDeltaModelInstance_InputBuffer"), ERDGBufferFlags::None);
+			RDGInputBuffer = Builder.ConvertToExternalBuffer(RDGBuffer);
 
 			Builder.Execute();
 		}
diff --git a/Engine/Plugins/Animation/MLDeformer/VertexDeltaModel/Source/VertexDeltaModel/Public/VertexDeltaModelInstance.h b/Engine/Plugins/Animation/MLDeformer/VertexDeltaModel/Source/VertexDeltaModel/Public/VertexDeltaModelInstance.h
--- a/Engine/Plugins/Animation/MLDeformer/VertexDeltaModel/Source/VertexDeltaModel/Public/VertexDeltaModelInstance.h
+++ b/Engine/Plugins/Animation/MLDeformer/VertexDeltaModel/Source/VertexDeltaModel/Public/VertexDeltaModelInstance.h
@@ -40,7 +40,7 @@ public:
 	 * Get the output vertex delta buffer
 	 * @return The FRDGBuffer for vertex deltas
 	 */
-	FRDGBuffer* GetOutputRDGBuffer() const;
+	TRefCountPtr<FRDGPooledBuffer> GetOutputRDGBuffer() const;
 
 	/**
 	 * Get the render graph buffer description required for the output of a neural network. Return false if a flat buffer is not appropriate
@@ -54,10 +54,10 @@ private:
 	void CreateNNEModel();
 
 	// Input Buffer for Joint Matrices / Curve Floats
-	FRDGBuffer* RDGInputBuffer = nullptr;
+	TRefCountPtr<FRDGPooledBuffer> RDGInputBuffer;
 
 	// Output Buffer for Vertex Deltas
-	FRDGBuffer* RDGVertexDeltaBuffer = nullptr;
+	TRefCountPtr<FRDGPooledBuffer> RDGVertexDeltaBuffer;
 
 	// The NNE RDG Model 
 	TUniquePtr<UE::NNECore::IModelRDG> ModelRDG;
