MLDeformer: Fix Non-Unity build errors in VertexDeltaModel
#rnx

diff --git a/Engine/Plugins/Animation/MLDeformer/VertexDeltaModel/Source/VertexDeltaModel/Private/VertexDeltaModelInstance.cpp b/Engine/Plugins/Animation/MLDeformer/VertexDeltaModel/Source/VertexDeltaModel/Private/VertexDeltaModelInstance.cpp
--- a/Engine/Plugins/Animation/MLDeformer/VertexDeltaModel/Source/VertexDeltaModel/Private/VertexDeltaModelInstance.cpp
+++ b/Engine/Plugins/Animation/MLDeformer/VertexDeltaModel/Source/VertexDeltaModel/Private/VertexDeltaModelInstance.cpp
@@ -4,6 +4,7 @@
 #include "VertexDeltaModel.h"
 #include "MLDeformerAsset.h"
 #include "RenderGraphBuilder.h"
+#include "RenderGraphResources.h"
 #include "Components/SkeletalMeshComponent.h"
 
 
diff --git a/Engine/Plugins/Animation/MLDeformer/VertexDeltaModel/Source/VertexDeltaModel/Public/VertexDeltaModelInstance.h b/Engine/Plugins/Animation/MLDeformer/VertexDeltaModel/Source/VertexDeltaModel/Public/VertexDeltaModelInstance.h
--- a/Engine/Plugins/Animation/MLDeformer/VertexDeltaModel/Source/VertexDeltaModel/Public/VertexDeltaModelInstance.h
+++ b/Engine/Plugins/Animation/MLDeformer/VertexDeltaModel/Source/VertexDeltaModel/Public/VertexDeltaModelInstance.h
@@ -11,6 +11,7 @@
 
 class UVertexDeltaModel;
 class UNeuralNetwork;
+class FRDGPooledBuffer;
 
 UCLASS()
 class VERTEXDELTAMODEL_API UVertexDeltaModelInstance
