MLDeformer:  Fix Non-Unity build errors in vertex delta model
#rnx

diff --git a/Engine/Plugins/Animation/MLDeformer/VertexDeltaModel/Source/VertexDeltaModel/Private/VertexDeltaModelInstance.cpp b/Engine/Plugins/Animation/MLDeformer/VertexDeltaModel/Source/VertexDeltaModel/Private/VertexDeltaModelInstance.cpp
--- a/Engine/Plugins/Animation/MLDeformer/VertexDeltaModel/Source/VertexDeltaModel/Private/VertexDeltaModelInstance.cpp
+++ b/Engine/Plugins/Animation/MLDeformer/VertexDeltaModel/Source/VertexDeltaModel/Private/VertexDeltaModelInstance.cpp
@@ -4,10 +4,8 @@
 #include "VertexDeltaModel.h"
 #include "MLDeformerAsset.h"
 #include "RenderGraphBuilder.h"
-#include "RenderGraphResources.h"
 #include "Components/SkeletalMeshComponent.h"
 
-
 UVertexDeltaModel* UVertexDeltaModelInstance::GetVertexDeltaModel() const
 {
 	return Cast<UVertexDeltaModel>(Model);
diff --git a/Engine/Plugins/Animation/MLDeformer/VertexDeltaModel/Source/VertexDeltaModel/Public/VertexDeltaModelInstance.h b/Engine/Plugins/Animation/MLDeformer/VertexDeltaModel/Source/VertexDeltaModel/Public/VertexDeltaModelInstance.h
--- a/Engine/Plugins/Animation/MLDeformer/VertexDeltaModel/Source/VertexDeltaModel/Public/VertexDeltaModelInstance.h
+++ b/Engine/Plugins/Animation/MLDeformer/VertexDeltaModel/Source/VertexDeltaModel/Public/VertexDeltaModelInstance.h
@@ -7,11 +7,11 @@
 #include "NNECoreRuntime.h"
 #include "NNECoreModelData.h"
 #include "NNECoreRuntimeRDG.h"
+#include "RenderGraphResources.h"
 #include "VertexDeltaModelInstance.generated.h"
 
 class UVertexDeltaModel;
 class UNeuralNetwork;
-class FRDGPooledBuffer;
 
 UCLASS()
 class VERTEXDELTAMODEL_API UVertexDeltaModelInstance
