[MLDeformer]

diff --git a/Engine/Plugins/Animation/MLDeformer/VertexDeltaModel/Source/VertexDeltaModel/Private/VertexDeltaModel.cpp b/Engine/Plugins/Animation/MLDeformer/VertexDeltaModel/Source/VertexDeltaModel/Private/VertexDeltaModel.cpp
--- a/Engine/Plugins/Animation/MLDeformer/VertexDeltaModel/Source/VertexDeltaModel/Private/VertexDeltaModel.cpp
+++ b/Engine/Plugins/Animation/MLDeformer/VertexDeltaModel/Source/VertexDeltaModel/Private/VertexDeltaModel.cpp
@@ -43,7 +43,8 @@ UVertexDeltaModel::UVertexDeltaModel(const FObjectInitializer& ObjectInitializer
 
 void UVertexDeltaModel::SetNNEModelData(TObjectPtr<UNNEModelData> ModelData)
 {
-	NNEModel = ModelData; 
+	NNEModel = ModelData;
+	GetReinitModelInstanceDelegate().Broadcast();
 }
 
 UMLDeformerModelInstance* UVertexDeltaModel::CreateModelInstance(UMLDeformerComponent* Component)
