NNE Rename NNECore module to NNE following split of runtimes into there own plugins (step2)

What:
* Following the split of NNE into plugins it became apparent NNECore would be better renamed to NNE to follow UE coding standard, this CL is the second step where the namespace is renamed.

Tests:
* All nne tests on win64 editor + style transfer demo

#rnx

diff --git a/Engine/Plugins/AI/MLAdapter/Source/MLAdapter/Private/Agents/MLAdapterAgent_Inference.cpp b/Engine/Plugins/AI/MLAdapter/Source/MLAdapter/Private/Agents/MLAdapterAgent_Inference.cpp
--- a/Engine/Plugins/AI/MLAdapter/Source/MLAdapter/Private/Agents/MLAdapterAgent_Inference.cpp
+++ b/Engine/Plugins/AI/MLAdapter/Source/MLAdapter/Private/Agents/MLAdapterAgent_Inference.cpp
@@ -18,7 +18,7 @@ void UMLAdapterAgent_Inference::PostInitProperties()
 		return;
 	}
 
-	TWeakInterfacePtr<INNERuntimeCPU> Runtime = UE::NNECore::GetRuntime<INNERuntimeCPU>(RuntimeName);
+	TWeakInterfacePtr<INNERuntimeCPU> Runtime = UE::NNE::GetRuntime<INNERuntimeCPU>(RuntimeName);
 	if (!Runtime.IsValid())
 	{
 		UE_LOG(LogMLAdapter, Warning, TEXT("Runtime %s is not valid!"), *RuntimeName);
@@ -27,7 +27,7 @@ void UMLAdapterAgent_Inference::PostInitProperties()
 
 	Brain = Runtime->CreateModel(ModelData)->CreateModelInstance();
 
-	TConstArrayView<UE::NNECore::FTensorDesc> InputTensorDescs = Brain->GetInputTensorDescs();
+	TConstArrayView<UE::NNE::FTensorDesc> InputTensorDescs = Brain->GetInputTensorDescs();
 	if (InputTensorDescs.Num() != 1)
 	{
 		UE_LOG(LogMLAdapter, Warning, TEXT("Brain accepts only single tensor input."));
@@ -39,18 +39,18 @@ void UMLAdapterAgent_Inference::PostInitProperties()
 		return;
 	}
 
-	UE::NNECore::FTensorShape InputTensorShape = UE::NNECore::FTensorShape::MakeFromSymbolic(InputTensorDescs[0].GetShape());
+	UE::NNE::FTensorShape InputTensorShape = UE::NNE::FTensorShape::MakeFromSymbolic(InputTensorDescs[0].GetShape());
 
 	Brain->SetInputTensorShapes({ InputTensorShape });
 
-	TConstArrayView<UE::NNECore::FTensorDesc> OutputTensorDescs = Brain->GetOutputTensorDescs();
+	TConstArrayView<UE::NNE::FTensorDesc> OutputTensorDescs = Brain->GetOutputTensorDescs();
 	if (OutputTensorDescs.Num() != 1)
 	{
 		UE_LOG(LogMLAdapter, Warning, TEXT("Brain accepts only single tensor output."));
 		return;
 	}
 
-	TConstArrayView<UE::NNECore::FTensorShape> OutputTensorShapes = Brain->GetOutputTensorShapes();
+	TConstArrayView<UE::NNE::FTensorShape> OutputTensorShapes = Brain->GetOutputTensorShapes();
 	if (OutputTensorShapes.Num() != 1)
 	{
 		UE_LOG(LogMLAdapter, Warning, TEXT("Brain could not resolve output shapes."));
@@ -106,8 +106,8 @@ void UMLAdapterAgent_Inference::Think(const float DeltaTime)
 	TArray<uint8> OutputBuffer;
 	OutputBuffer.SetNumUninitialized(OutputTensorSizeInBytes);
 
-	UE::NNECore::FTensorBindingCPU InputTensorBinding{Buffer.GetData(), Buffer.Num()};
-	UE::NNECore::FTensorBindingCPU OutputTensorBinding{OutputBuffer.GetData(), OutputBuffer.Num()};
+	UE::NNE::FTensorBindingCPU InputTensorBinding{Buffer.GetData(), Buffer.Num()};
+	UE::NNE::FTensorBindingCPU OutputTensorBinding{OutputBuffer.GetData(), OutputBuffer.Num()};
 
 	Brain->RunSync({InputTensorBinding}, {OutputTensorBinding});
 
diff --git a/Engine/Plugins/AI/MLAdapter/Source/MLAdapter/Public/Agents/MLAdapterAgent_Inference.h b/Engine/Plugins/AI/MLAdapter/Source/MLAdapter/Public/Agents/MLAdapterAgent_Inference.h
--- a/Engine/Plugins/AI/MLAdapter/Source/MLAdapter/Public/Agents/MLAdapterAgent_Inference.h
+++ b/Engine/Plugins/AI/MLAdapter/Source/MLAdapter/Public/Agents/MLAdapterAgent_Inference.h
@@ -30,7 +30,7 @@ public:
 	TObjectPtr<UNNEModelData> ModelData;
 
 	/** The neural network controlling this agent. */
-	TUniquePtr<UE::NNECore::IModelInstanceCPU> Brain;
+	TUniquePtr<UE::NNE::IModelInstanceCPU> Brain;
 
 private:
 	uint32 InputTensorSizeInBytes = 0;
diff --git a/Engine/Plugins/Animation/MLDeformer/VertexDeltaModel/Source/VertexDeltaModel/Private/VertexDeltaModelInstance.cpp b/Engine/Plugins/Animation/MLDeformer/VertexDeltaModel/Source/VertexDeltaModel/Private/VertexDeltaModelInstance.cpp
--- a/Engine/Plugins/Animation/MLDeformer/VertexDeltaModel/Source/VertexDeltaModel/Private/VertexDeltaModelInstance.cpp
+++ b/Engine/Plugins/Animation/MLDeformer/VertexDeltaModel/Source/VertexDeltaModel/Private/VertexDeltaModelInstance.cpp
@@ -12,7 +12,7 @@ UVertexDeltaModel* UVertexDeltaModelInstance::GetVertexDeltaModel() const
 }
 
 
-UE::NNECore::IModelInstanceRDG* UVertexDeltaModelInstance::GetNNEModelInstanceRDG() const
+UE::NNE::IModelInstanceRDG* UVertexDeltaModelInstance::GetNNEModelInstanceRDG() const
 {
 	return ModelInstanceRDG.Get();
 }
@@ -27,10 +27,10 @@ FString UVertexDeltaModelInstance::CheckCompatibility(USkeletalMeshComponent* In
 	FString ErrorString = Super::CheckCompatibility(InSkelMeshComponent, bLogIssues);
 
 	// Verify the number of inputs versus the expected number of inputs.	
-	const UE::NNECore::IModelInstanceRDG* RDGModelInstance = GetNNEModelInstanceRDG();
+	const UE::NNE::IModelInstanceRDG* RDGModelInstance = GetNNEModelInstanceRDG();
 	if (RDGModelInstance && Model->GetDeformerAsset())
 	{
-		TConstArrayView<UE::NNECore::FTensorDesc> GPUTensorDesc = RDGModelInstance->GetInputTensorDescs();
+		TConstArrayView<UE::NNE::FTensorDesc> GPUTensorDesc = RDGModelInstance->GetInputTensorDescs();
 		const int64 NumNeuralNetInputs = GPUTensorDesc[0].GetShape().Rank();
 		const int64 NumDeformerAssetInputs = static_cast<int64>(Model->GetInputInfo()->CalcNumNeuralNetInputs(Model->GetNumFloatsPerBone(), Model->GetNumFloatsPerCurve()));
 		if (NumNeuralNetInputs != NumDeformerAssetInputs)
@@ -72,14 +72,14 @@ void UVertexDeltaModelInstance::Execute(float ModelWeight)
 					FRDGBuilder GraphBuilder(RHICmdList);
 
 					// Build Input Bindings
-					TArray<UE::NNECore::FTensorBindingRDG> InputBindingsRDG;
-					UE::NNECore::FTensorBindingRDG& BindingRDG = InputBindingsRDG.Emplace_GetRef();
+					TArray<UE::NNE::FTensorBindingRDG> InputBindingsRDG;
+					UE::NNE::FTensorBindingRDG& BindingRDG = InputBindingsRDG.Emplace_GetRef();
 					BindingRDG.Buffer = GraphBuilder.RegisterExternalBuffer(RDGInputBuffer);
 					GraphBuilder.QueueBufferUpload(BindingRDG.Buffer, NNEInputTensorBuffer.GetData(), NNEInputTensorBuffer.Num() * sizeof(float), ERDGInitialDataFlags::NoCopy);
 	
 					// Build Output Bindings
-					TArray<UE::NNECore::FTensorBindingRDG> OutputBindingsRDG;
-					UE::NNECore::FTensorBindingRDG& OutputBindingRDG = OutputBindingsRDG.Emplace_GetRef();
+					TArray<UE::NNE::FTensorBindingRDG> OutputBindingsRDG;
+					UE::NNE::FTensorBindingRDG& OutputBindingRDG = OutputBindingsRDG.Emplace_GetRef();
 					OutputBindingRDG.Buffer = GraphBuilder.RegisterExternalBuffer(RDGVertexDeltaBuffer);
 
 					if (!ModelInstanceRDG.IsValid())
@@ -94,12 +94,12 @@ void UVertexDeltaModelInstance::Execute(float ModelWeight)
 	}
 }
 
-bool UVertexDeltaModelInstance::GetRDGVertexBufferDesc(TConstArrayView<UE::NNECore::FTensorDesc>& InOutputTensorDescs, FRDGBufferDesc& OutBufferDesc)
+bool UVertexDeltaModelInstance::GetRDGVertexBufferDesc(TConstArrayView<UE::NNE::FTensorDesc>& InOutputTensorDescs, FRDGBufferDesc& OutBufferDesc)
 {
 	if (InOutputTensorDescs.Num() > 0)
 	{
 		const uint32 ElemByteSize = InOutputTensorDescs[0].GetElemByteSize();
-		const UE::NNECore::FSymbolicTensorShape& SymShape = InOutputTensorDescs[0].GetShape();
+		const UE::NNE::FSymbolicTensorShape& SymShape = InOutputTensorDescs[0].GetShape();
 		for (int32 i = 1; i < InOutputTensorDescs.Num(); i++)
 		{
 			if (InOutputTensorDescs[i].GetElemByteSize() != ElemByteSize || SymShape != InOutputTensorDescs[i].GetShape())
@@ -108,7 +108,7 @@ bool UVertexDeltaModelInstance::GetRDGVertexBufferDesc(TConstArrayView<UE::NNECo
 			}
 		}
 		// Create a single flat output buffer
-		const UE::NNECore::FTensorShape OutputShape = UE::NNECore::FTensorShape::MakeFromSymbolic(SymShape);
+		const UE::NNE::FTensorShape OutputShape = UE::NNE::FTensorShape::MakeFromSymbolic(SymShape);
 		OutBufferDesc.BytesPerElement = ElemByteSize;
 		OutBufferDesc.NumElements = OutputShape.Volume() * InOutputTensorDescs.Num();
 		OutBufferDesc.Usage = EBufferUsageFlags::UnorderedAccess | EBufferUsageFlags::ShaderResource | EBufferUsageFlags::VertexBuffer;
@@ -117,7 +117,7 @@ bool UVertexDeltaModelInstance::GetRDGVertexBufferDesc(TConstArrayView<UE::NNECo
 	return false;
 }
 
-void UVertexDeltaModelInstance::CreateRDGBuffers(TConstArrayView<UE::NNECore::FTensorDesc>& OutputTensorDescs)
+void UVertexDeltaModelInstance::CreateRDGBuffers(TConstArrayView<UE::NNE::FTensorDesc>& OutputTensorDescs)
 {
 	ENQUEUE_RENDER_COMMAND(VertexDeltaModelInstance_CreateOuputRDGBuffer)(
 		[this, &OutputTensorDescs](FRHICommandListImmediate& RHICmdList)
@@ -162,8 +162,8 @@ void UVertexDeltaModelInstance::CreateNNEModel()
 		UVertexDeltaModel* VertexDeltaModel = Cast<UVertexDeltaModel>(Model);
 		if (VertexDeltaModel)
 		{
-			TWeakInterfacePtr<INNERuntime> Runtime = UE::NNECore::GetRuntime<INNERuntime>(VertexDeltaModel->GetNNERuntimeName());
-			TWeakInterfacePtr<INNERuntimeRDG> RuntimeRDG = UE::NNECore::GetRuntime<INNERuntimeRDG>(VertexDeltaModel->GetNNERuntimeName());
+			TWeakInterfacePtr<INNERuntime> Runtime = UE::NNE::GetRuntime<INNERuntime>(VertexDeltaModel->GetNNERuntimeName());
+			TWeakInterfacePtr<INNERuntimeRDG> RuntimeRDG = UE::NNE::GetRuntime<INNERuntimeRDG>(VertexDeltaModel->GetNNERuntimeName());
 
 			if (!Runtime.IsValid())
 			{
@@ -182,14 +182,14 @@ void UVertexDeltaModelInstance::CreateNNEModel()
 					if (ModelInstanceRDG)
 					{
 						// setup inputs
-						TConstArrayView<UE::NNECore::FTensorDesc> InputTensorDescs = ModelInstanceRDG->GetInputTensorDescs();
-						UE::NNECore::FTensorShape InputTensorShape = UE::NNECore::FTensorShape::MakeFromSymbolic(InputTensorDescs[0].GetShape());
+						TConstArrayView<UE::NNE::FTensorDesc> InputTensorDescs = ModelInstanceRDG->GetInputTensorDescs();
+						UE::NNE::FTensorShape InputTensorShape = UE::NNE::FTensorShape::MakeFromSymbolic(InputTensorDescs[0].GetShape());
 						ModelInstanceRDG->SetInputTensorShapes({ InputTensorShape });
 						check(InputTensorDescs[0].GetElemByteSize() == sizeof(float));
 						NNEInputTensorBuffer.SetNumUninitialized(InputTensorShape.Volume());
 
 						// setup outputs
-						TConstArrayView<UE::NNECore::FTensorDesc> OutputTensorDescs = ModelInstanceRDG->GetOutputTensorDescs();
+						TConstArrayView<UE::NNE::FTensorDesc> OutputTensorDescs = ModelInstanceRDG->GetOutputTensorDescs();
 						CreateRDGBuffers(OutputTensorDescs);
 					}
 					else
@@ -223,8 +223,8 @@ bool UVertexDeltaModelInstance::SetupInputs()
 		return false;
 	}
 
-	TConstArrayView<UE::NNECore::FTensorDesc> InputTensorDescs = ModelInstanceRDG->GetInputTensorDescs();
-	const UE::NNECore::FTensorShape InputTensorShape = UE::NNECore::FTensorShape::MakeFromSymbolic(InputTensorDescs[0].GetShape());
+	TConstArrayView<UE::NNE::FTensorDesc> InputTensorDescs = ModelInstanceRDG->GetInputTensorDescs();
+	const UE::NNE::FTensorShape InputTensorShape = UE::NNE::FTensorShape::MakeFromSymbolic(InputTensorDescs[0].GetShape());
 	const int64 NumNeuralNetInputs = InputTensorShape.Volume();
 	const int64 NumDeformerAssetInputs = Model->GetInputInfo()->CalcNumNeuralNetInputs(Model->GetNumFloatsPerBone(), Model->GetNumFloatsPerCurve());
 	if (NumNeuralNetInputs != NumDeformerAssetInputs)
diff --git a/Engine/Plugins/Animation/MLDeformer/VertexDeltaModel/Source/VertexDeltaModel/Public/VertexDeltaGraphDataInterface.h b/Engine/Plugins/Animation/MLDeformer/VertexDeltaModel/Source/VertexDeltaModel/Public/VertexDeltaGraphDataInterface.h
--- a/Engine/Plugins/Animation/MLDeformer/VertexDeltaModel/Source/VertexDeltaModel/Public/VertexDeltaGraphDataInterface.h
+++ b/Engine/Plugins/Animation/MLDeformer/VertexDeltaModel/Source/VertexDeltaModel/Public/VertexDeltaGraphDataInterface.h
@@ -85,7 +85,7 @@ namespace UE::VertexDeltaModel
 
 	protected:
 		FSkeletalMeshObject* SkeletalMeshObject = nullptr;
-		UE::NNECore::IModelInstanceRDG* NNEModelInstanceRDG = nullptr;
+		UE::NNE::IModelInstanceRDG* NNEModelInstanceRDG = nullptr;
 		UVertexDeltaModelInstance* VertexDeltaModelInstance = nullptr;
 		FRHIShaderResourceView* VertexMapBufferSRV = nullptr;
 		FRDGBuffer* Buffer = nullptr;
diff --git a/Engine/Plugins/Animation/MLDeformer/VertexDeltaModel/Source/VertexDeltaModel/Public/VertexDeltaModelInstance.h b/Engine/Plugins/Animation/MLDeformer/VertexDeltaModel/Source/VertexDeltaModel/Public/VertexDeltaModelInstance.h
--- a/Engine/Plugins/Animation/MLDeformer/VertexDeltaModel/Source/VertexDeltaModel/Public/VertexDeltaModelInstance.h
+++ b/Engine/Plugins/Animation/MLDeformer/VertexDeltaModel/Source/VertexDeltaModel/Public/VertexDeltaModelInstance.h
@@ -34,7 +34,7 @@ public:
 	 * Get the RDG Neural Network Model Instance
 	 * @return The IModelInstanceRDG if successfully created, or nullptr
 	 */
-	UE::NNECore::IModelInstanceRDG* GetNNEModelInstanceRDG() const;
+	UE::NNE::IModelInstanceRDG* GetNNEModelInstanceRDG() const;
 
 	/**
 	 * Get the output vertex delta buffer
@@ -46,11 +46,11 @@ public:
 	 * Get the render graph buffer description required for the output of a neural network. Return false if a flat buffer is not appropriate
 	 * @return True if the OutputTensorDescs can be represented by a flat float  FRDGBuffer
 	 */
-	bool GetRDGVertexBufferDesc(TConstArrayView<UE::NNECore::FTensorDesc>& InOutputTensorDescs, FRDGBufferDesc& OutBufferDesc);
+	bool GetRDGVertexBufferDesc(TConstArrayView<UE::NNE::FTensorDesc>& InOutputTensorDescs, FRDGBufferDesc& OutBufferDesc);
 
 private:
 
-	void CreateRDGBuffers(TConstArrayView<UE::NNECore::FTensorDesc>& OutputTensorDescs);
+	void CreateRDGBuffers(TConstArrayView<UE::NNE::FTensorDesc>& OutputTensorDescs);
 	void CreateNNEModel();
 
 	// Input Buffer for Joint Matrices / Curve Floats
@@ -60,7 +60,7 @@ private:
 	TRefCountPtr<FRDGPooledBuffer> RDGVertexDeltaBuffer;
 
 	// The NNE RDG Model 
-	TUniquePtr<UE::NNECore::IModelInstanceRDG> ModelInstanceRDG;
+	TUniquePtr<UE::NNE::IModelInstanceRDG> ModelInstanceRDG;
 	
 	// The CPU Input Tensor Buffer
 	TArray<float> NNEInputTensorBuffer;
diff --git a/Engine/Plugins/Experimental/NNE/Source/NNE/Internal/NNEAttributeMap.h b/Engine/Plugins/Experimental/NNE/Source/NNE/Internal/NNEAttributeMap.h
--- a/Engine/Plugins/Experimental/NNE/Source/NNE/Internal/NNEAttributeMap.h
+++ b/Engine/Plugins/Experimental/NNE/Source/NNE/Internal/NNEAttributeMap.h
@@ -4,7 +4,7 @@
 
 #include "NNEAttributeValue.h"
 
-namespace UE::NNECore
+namespace UE::NNE
 {
 
 class FAttributeMap
@@ -77,4 +77,4 @@ private:
 	TArray<FEntry> Attributes;
 };
 
-} // namespace UE::NNECore
\ No newline at end of file
+} // namespace UE::NNE
\ No newline at end of file
diff --git a/Engine/Plugins/Experimental/NNE/Source/NNE/Internal/NNEModelBase.h b/Engine/Plugins/Experimental/NNE/Source/NNE/Internal/NNEModelBase.h
--- a/Engine/Plugins/Experimental/NNE/Source/NNE/Internal/NNEModelBase.h
+++ b/Engine/Plugins/Experimental/NNE/Source/NNE/Internal/NNEModelBase.h
@@ -6,7 +6,7 @@
 #include "NNE.h"
 #include "NNETypes.h"
 
-namespace UE::NNECore::Internal
+namespace UE::NNE::Internal
 {
 	template <class ModelInterface> class FModelInstanceBase : public ModelInterface
 	{
@@ -14,11 +14,11 @@ namespace UE::NNECore::Internal
 
 		virtual ~FModelInstanceBase() = default;
 
-		virtual TConstArrayView<NNECore::FTensorDesc> GetInputTensorDescs() const override;
-		virtual TConstArrayView<NNECore::FTensorDesc> GetOutputTensorDescs() const override;
-		virtual TConstArrayView<NNECore::FTensorShape> GetInputTensorShapes() const override;
-		virtual TConstArrayView<NNECore::FTensorShape> GetOutputTensorShapes() const override;
-		virtual int SetInputTensorShapes(TConstArrayView<NNECore::FTensorShape> InInputShapes) override;
+		virtual TConstArrayView<NNE::FTensorDesc> GetInputTensorDescs() const override;
+		virtual TConstArrayView<NNE::FTensorDesc> GetOutputTensorDescs() const override;
+		virtual TConstArrayView<NNE::FTensorShape> GetInputTensorShapes() const override;
+		virtual TConstArrayView<NNE::FTensorShape> GetOutputTensorShapes() const override;
+		virtual int SetInputTensorShapes(TConstArrayView<NNE::FTensorShape> InInputShapes) override;
 
 	protected:
 
@@ -77,4 +77,4 @@ namespace UE::NNECore::Internal
 		return 0;
 	}
 
-} // namespace UE::NNECore::Internal
\ No newline at end of file
+} // namespace UE::NNE::Internal
\ No newline at end of file
diff --git a/Engine/Plugins/Experimental/NNE/Source/NNE/Internal/NNEModelOptimizerInterface.h b/Engine/Plugins/Experimental/NNE/Source/NNE/Internal/NNEModelOptimizerInterface.h
--- a/Engine/Plugins/Experimental/NNE/Source/NNE/Internal/NNEModelOptimizerInterface.h
+++ b/Engine/Plugins/Experimental/NNE/Source/NNE/Internal/NNEModelOptimizerInterface.h
@@ -3,14 +3,14 @@
 #pragma once
 
 struct FNNEModelRaw;
-namespace UE::NNECore { class FAttributeMap; }
+namespace UE::NNE { class FAttributeMap; }
 
-namespace UE::NNECore::Internal
+namespace UE::NNE::Internal
 {
 	
 //Note: create a dedicated FOptimizerOptionsMap when function diverge with FAttributeMap
 //example introduction of sparse tensor to FAttribute witch make no sense as an optimizer attribute
-using FOptimizerOptionsMap = UE::NNECore::FAttributeMap;
+using FOptimizerOptionsMap = UE::NNE::FAttributeMap;
 
 /** Interface class for NNE model validator */
 class IModelValidator
@@ -49,4 +49,4 @@ public:
 	virtual bool Optimize(const FNNEModelRaw& InputModel, FNNEModelRaw& OutModel, const FOptimizerOptionsMap& Options) = 0;
 };
 
-} // namespace UE::NNECore::Internal
+} // namespace UE::NNE::Internal
diff --git a/Engine/Plugins/Experimental/NNE/Source/NNE/Internal/NNEOperator.h b/Engine/Plugins/Experimental/NNE/Source/NNE/Internal/NNEOperator.h
--- a/Engine/Plugins/Experimental/NNE/Source/NNE/Internal/NNEOperator.h
+++ b/Engine/Plugins/Experimental/NNE/Source/NNE/Internal/NNEOperator.h
@@ -4,7 +4,7 @@
 
 #include "HAL/Platform.h"
 
-namespace UE::NNECore::Internal
+namespace UE::NNE::Internal
 {
 
 	//One input element-wise operators
@@ -86,4 +86,4 @@ namespace UE::NNECore::Internal
 		MAX
 	};
 
-} // namespace UE::NNECore::Internal
\ No newline at end of file
+} // namespace UE::NNE::Internal
\ No newline at end of file
diff --git a/Engine/Plugins/Experimental/NNE/Source/NNE/Internal/NNETensor.h b/Engine/Plugins/Experimental/NNE/Source/NNE/Internal/NNETensor.h
--- a/Engine/Plugins/Experimental/NNE/Source/NNE/Internal/NNETensor.h
+++ b/Engine/Plugins/Experimental/NNE/Source/NNE/Internal/NNETensor.h
@@ -5,7 +5,7 @@
 #include "CoreMinimal.h"
 #include "NNETypes.h"
 
-namespace UE::NNECore::Internal
+namespace UE::NNE::Internal
 {
 	/** Concrete tensor with data accessible by graph scheduling */
 	class FTensor
@@ -107,4 +107,4 @@ namespace UE::NNECore::Internal
 
 	using FTensorRef = FTensor*;
 
-} // namespace UE::NNECore::Internal
\ No newline at end of file
+} // namespace UE::NNE::Internal
\ No newline at end of file
diff --git a/Engine/Plugins/Experimental/NNE/Source/NNE/Private/NNE.cpp b/Engine/Plugins/Experimental/NNE/Source/NNE/Private/NNE.cpp
--- a/Engine/Plugins/Experimental/NNE/Source/NNE/Private/NNE.cpp
+++ b/Engine/Plugins/Experimental/NNE/Source/NNE/Private/NNE.cpp
@@ -9,7 +9,7 @@
 
 DEFINE_LOG_CATEGORY(LogNNE);
 
-namespace UE::NNECore
+namespace UE::NNE
 {
 	class FRegistry
 	{
diff --git a/Engine/Plugins/Experimental/NNE/Source/NNE/Private/NNEModelData.cpp b/Engine/Plugins/Experimental/NNE/Source/NNE/Private/NNEModelData.cpp
--- a/Engine/Plugins/Experimental/NNE/Source/NNE/Private/NNEModelData.cpp
+++ b/Engine/Plugins/Experimental/NNE/Source/NNE/Private/NNEModelData.cpp
@@ -71,7 +71,7 @@ inline void PutIntoDDC(const FGuid& FileDataId, const FString& RuntimeName, FSha
 
 inline TArray<uint8> CreateRuntimeDataBlob(const FString& RuntimeName, FString FileType, const TArray<uint8>& FileData)
 {
-	TWeakInterfacePtr<INNERuntime> NNERuntime = UE::NNECore::GetRuntime<INNERuntime>(RuntimeName);
+	TWeakInterfacePtr<INNERuntime> NNERuntime = UE::NNE::GetRuntime<INNERuntime>(RuntimeName);
 	if (NNERuntime.IsValid())
 	{
 		return NNERuntime->CreateModelData(FileType, FileData);
@@ -79,7 +79,7 @@ inline TArray<uint8> CreateRuntimeDataBlob(const FString& RuntimeName, FString F
 	else
 	{
 		UE_LOG(LogNNE, Error, TEXT("UNNEModelData: No runtime '%s' found. Valid runtimes are: "), *RuntimeName);
-		TArrayView<TWeakInterfacePtr<INNERuntime>> Runtimes = UE::NNECore::GetAllRuntimes();
+		TArrayView<TWeakInterfacePtr<INNERuntime>> Runtimes = UE::NNE::GetAllRuntimes();
 		for (int i = 0; i < Runtimes.Num(); i++)
 		{
 			UE_LOG(LogNNE, Error, TEXT("- %s"), *Runtimes[i]->GetRuntimeName());
@@ -168,7 +168,7 @@ void UNNEModelData::Serialize(FArchive& Ar)
 		//No target runtime means all currently registered ones.
 		if (GetTargetRuntimes().IsEmpty())
 		{
-			for (const TWeakInterfacePtr<INNERuntime>& Runtime : UE::NNECore::GetAllRuntimes())
+			for (const TWeakInterfacePtr<INNERuntime>& Runtime : UE::NNE::GetAllRuntimes())
 			{
 				CookedRuntimeNames.Add(Runtime->GetRuntimeName());
 			}
@@ -242,7 +242,7 @@ void UNNEModelData::Serialize(FArchive& Ar)
 }
 
 #if WITH_EDITORONLY_DATA
-namespace UE::NNECore::ModelDataHelpers
+namespace UE::NNE::ModelDataHelpers
 {
 	FString GetRuntimesAsString(TArrayView<const FString> Runtimes)
 	{
@@ -265,7 +265,7 @@ namespace UE::NNECore::ModelDataHelpers
 		}
 		return RuntimesAsOneString;
 	}
-} // UE::NNECore::ModelDataHelpers
+} // UE::NNE::ModelDataHelpers
 
 void UNNEModelData::PostInitProperties()
 {
@@ -283,7 +283,7 @@ void UNNEModelData::GetAssetRegistryTags(TArray<FAssetRegistryTag>& OutTags) con
 		OutTags.Add(FAssetRegistryTag(SourceFileTagName(), AssetImportData->GetSourceData().ToJson(), FAssetRegistryTag::TT_Hidden));
 	}
 
-	OutTags.Add(FAssetRegistryTag("TargetRuntimes", UE::NNECore::ModelDataHelpers::GetRuntimesAsString(GetTargetRuntimes()), FAssetRegistryTag::TT_Alphabetical));
+	OutTags.Add(FAssetRegistryTag("TargetRuntimes", UE::NNE::ModelDataHelpers::GetRuntimesAsString(GetTargetRuntimes()), FAssetRegistryTag::TT_Alphabetical));
 
 	Super::GetAssetRegistryTags(OutTags);
 }
diff --git a/Engine/Plugins/Experimental/NNE/Source/NNE/Private/NNEModule.cpp b/Engine/Plugins/Experimental/NNE/Source/NNE/Private/NNEModule.cpp
--- a/Engine/Plugins/Experimental/NNE/Source/NNE/Private/NNEModule.cpp
+++ b/Engine/Plugins/Experimental/NNE/Source/NNE/Private/NNEModule.cpp
@@ -5,7 +5,7 @@
 #include "Misc/CoreDelegates.h"
 #include "Modules/ModuleManager.h"
 
-class FNNECoreModule : public IModuleInterface
+class FNNEModule : public IModuleInterface
 {
 public:
 
@@ -29,4 +29,4 @@ public:
 
 };
 
-IMPLEMENT_MODULE(FNNECoreModule, NNECore)
+IMPLEMENT_MODULE(FNNEModule, NNE)
diff --git a/Engine/Plugins/Experimental/NNE/Source/NNE/Private/NNETypes.cpp b/Engine/Plugins/Experimental/NNE/Source/NNE/Private/NNETypes.cpp
--- a/Engine/Plugins/Experimental/NNE/Source/NNE/Private/NNETypes.cpp
+++ b/Engine/Plugins/Experimental/NNE/Source/NNE/Private/NNETypes.cpp
@@ -3,7 +3,7 @@
 #include "NNETypes.h"
 #include "NNE.h"
 
-namespace UE::NNECore::TypesUtils
+namespace UE::NNE::TypesUtils
 {
 	template <class T> bool AreShapeEqual(const T& lhs, const T& rhs)
 	{
@@ -22,7 +22,7 @@ namespace UE::NNECore::TypesUtils
 	}
 }
 
-namespace UE::NNECore
+namespace UE::NNE
 {
 
 bool FSymbolicTensorShape::IsConcrete() const
@@ -173,4 +173,4 @@ FTensorDesc FTensorDesc::Make(const FString& Name, const FSymbolicTensorShape& S
 	return Desc;
 }
 
-} // namespace UE::NNECore
+} // namespace UE::NNE
diff --git a/Engine/Plugins/Experimental/NNE/Source/NNE/Public/NNE.h b/Engine/Plugins/Experimental/NNE/Source/NNE/Public/NNE.h
--- a/Engine/Plugins/Experimental/NNE/Source/NNE/Public/NNE.h
+++ b/Engine/Plugins/Experimental/NNE/Source/NNE/Public/NNE.h
@@ -8,7 +8,7 @@
 
 NNE_API DECLARE_LOG_CATEGORY_EXTERN(LogNNE, Log, All);
 
-namespace UE::NNECore
+namespace UE::NNE
 {
 	/**
 	 * Register a runtime to make it accessible to NNE clients.
diff --git a/Engine/Plugins/Experimental/NNE/Source/NNE/Public/NNEModelData.h b/Engine/Plugins/Experimental/NNE/Source/NNE/Public/NNEModelData.h
--- a/Engine/Plugins/Experimental/NNE/Source/NNE/Public/NNEModelData.h
+++ b/Engine/Plugins/Experimental/NNE/Source/NNE/Public/NNEModelData.h
@@ -11,7 +11,7 @@
  *
  * Neural network models typically consist of a graph of operations and corresponding parameters as e.g. weights.
  * UNNEModelData assets store such model data as imported e.g. by the UNNEModelDataFactory class.
- * An INNERuntime object retrieved by UE::NNECore::GetRuntime<T>(const FString& Name) can be used to create an inferable neural network model.
+ * An INNERuntime object retrieved by UE::NNE::GetRuntime<T>(const FString& Name) can be used to create an inferable neural network model.
  */
 UCLASS(BlueprintType, Category = "NNE")
 class NNE_API UNNEModelData : public UObject
diff --git a/Engine/Plugins/Experimental/NNE/Source/NNE/Public/NNERuntime.h b/Engine/Plugins/Experimental/NNE/Source/NNE/Public/NNERuntime.h
--- a/Engine/Plugins/Experimental/NNE/Source/NNE/Public/NNERuntime.h
+++ b/Engine/Plugins/Experimental/NNE/Source/NNE/Public/NNERuntime.h
@@ -14,7 +14,7 @@ class NNE_API UNNERuntime : public UInterface
 };
 
 /**
- * The base interface of a neural network runtime. It is returned by UE::NNECore::GetAllRuntimes() giving access to all available runtimes.
+ * The base interface of a neural network runtime. It is returned by UE::NNE::GetAllRuntimes() giving access to all available runtimes.
  *
  * This interface is mainly used internally by UNNEModelData to cook a model given in a file format (e.g. .onnx) into runtime specific model data.
  * The model data is then stored inside UNNEModelData which then can be used in game to create an inferable model.
diff --git a/Engine/Plugins/Experimental/NNE/Source/NNE/Public/NNERuntimeCPU.h b/Engine/Plugins/Experimental/NNE/Source/NNE/Public/NNERuntimeCPU.h
--- a/Engine/Plugins/Experimental/NNE/Source/NNE/Public/NNERuntimeCPU.h
+++ b/Engine/Plugins/Experimental/NNE/Source/NNE/Public/NNERuntimeCPU.h
@@ -11,7 +11,7 @@
 
 class UNNEModelData;
 
-namespace UE::NNECore
+namespace UE::NNE
 {
 
 /**
@@ -28,8 +28,8 @@ struct NNE_API FTensorBindingCPU
 /**
  * The interface of a model instance that can run on CPU.
  *
- * Use UE::NNECore::IModelCPU::CreateModelInstance() to get a model instance.
- * Use UE::NNECore::GetRuntime<INNERuntimeCPU>(RuntimeName) to get a runtime capable of creating CPU models.
+ * Use UE::NNE::IModelCPU::CreateModelInstance() to get a model instance.
+ * Use UE::NNE::GetRuntime<INNERuntimeCPU>(RuntimeName) to get a runtime capable of creating CPU models.
  */
 class NNE_API IModelInstanceCPU
 {
@@ -100,7 +100,7 @@ public:
 /**
  * The interface of a model capable of creating model instance that can run on CPU.
  *
- * Use UE::NNECore::GetRuntime<INNERuntimeCPU>(RuntimeName) to get a runtime capable of creating CPU models.
+ * Use UE::NNE::GetRuntime<INNERuntimeCPU>(RuntimeName) to get a runtime capable of creating CPU models.
  */
 class NNE_API IModelCPU
 {
@@ -116,10 +116,10 @@ public:
 	 *
 	 * @return A caller owned model representing the neural network instance created.
 	 */
-	virtual TUniquePtr<UE::NNECore::IModelInstanceCPU> CreateModelInstance() = 0;
+	virtual TUniquePtr<UE::NNE::IModelInstanceCPU> CreateModelInstance() = 0;
 };
 
-} // UE::NNECore
+} // UE::NNE
 
 UINTERFACE()
 class NNE_API UNNERuntimeCPU : public UInterface
@@ -130,7 +130,7 @@ class NNE_API UNNERuntimeCPU : public UInterface
 /**
  * The interface of a neural network runtime capable of creating CPU models.
  *
- * Call UE::NNECore::GetRuntime<INNERuntimeCPU>(RuntimeName) to get a runtime implementing this interface.
+ * Call UE::NNE::GetRuntime<INNERuntimeCPU>(RuntimeName) to get a runtime implementing this interface.
  */
 class NNE_API INNERuntimeCPU
 {
@@ -156,5 +156,5 @@ public:
 	 * @param ModelData The model data for which to create a model.
 	 * @return A caller owned model representing the neural network created from ModelData.
 	 */
-	virtual TUniquePtr<UE::NNECore::IModelCPU> CreateModel(TObjectPtr<UNNEModelData> ModelData) = 0;
+	virtual TUniquePtr<UE::NNE::IModelCPU> CreateModel(TObjectPtr<UNNEModelData> ModelData) = 0;
 };
\ No newline at end of file
diff --git a/Engine/Plugins/Experimental/NNE/Source/NNE/Public/NNERuntimeGPU.h b/Engine/Plugins/Experimental/NNE/Source/NNE/Public/NNERuntimeGPU.h
--- a/Engine/Plugins/Experimental/NNE/Source/NNE/Public/NNERuntimeGPU.h
+++ b/Engine/Plugins/Experimental/NNE/Source/NNE/Public/NNERuntimeGPU.h
@@ -11,7 +11,7 @@
 
 class UNNEModelData;
 
-namespace UE::NNECore
+namespace UE::NNE
 {
 
 /**
@@ -28,8 +28,8 @@ struct NNE_API FTensorBindingGPU
 /**
  * The interface of a model instance that can run on GPU.
  *
- * Use UE::NNECore::IModelGPU::CreateModelInstance() to get a model instance.
- * Use UE::NNECore::GetRuntime<INNERuntimeGPU>(RuntimeName) to get a runtime capable of creating GPU models.
+ * Use UE::NNE::IModelGPU::CreateModelInstance() to get a model instance.
+ * Use UE::NNE::GetRuntime<INNERuntimeGPU>(RuntimeName) to get a runtime capable of creating GPU models.
  */
 class NNE_API IModelInstanceGPU
 {
@@ -100,7 +100,7 @@ public:
 /**
  * The interface of a model capable of creating model instance that can run on GPU.
  *
- * Use UE::NNECore::GetRuntime<INNERuntimeGPU>(RuntimeName) to get a runtime capable of creating GPU models.
+ * Use UE::NNE::GetRuntime<INNERuntimeGPU>(RuntimeName) to get a runtime capable of creating GPU models.
  */
 class NNE_API IModelGPU
 {
@@ -116,10 +116,10 @@ public:
 	 *
 	 * @return A caller owned model representing the neural network instance created.
 	 */
-	virtual TUniquePtr<UE::NNECore::IModelInstanceGPU> CreateModelInstance() = 0;
+	virtual TUniquePtr<UE::NNE::IModelInstanceGPU> CreateModelInstance() = 0;
 };
 
-} // UE::NNECore
+} // UE::NNE
 
 UINTERFACE()
 class NNE_API UNNERuntimeGPU : public UInterface
@@ -130,7 +130,7 @@ class NNE_API UNNERuntimeGPU : public UInterface
 /**
  * The interface of a neural network runtime capable of creating GPU models.
  *
- * Call UE::NNECore::GetRuntime<INNERuntimeGPU>(RuntimeName) to get a runtime implementing this interface.
+ * Call UE::NNE::GetRuntime<INNERuntimeGPU>(RuntimeName) to get a runtime implementing this interface.
  */
 class NNE_API INNERuntimeGPU
 {
@@ -156,5 +156,5 @@ public:
 	 * @param ModelData The model data for which to create a model.
 	 * @return A caller owned model representing the neural network created from ModelData.
 	 */
-	virtual TUniquePtr<UE::NNECore::IModelGPU> CreateModel(TObjectPtr<UNNEModelData> ModelData) = 0;
+	virtual TUniquePtr<UE::NNE::IModelGPU> CreateModel(TObjectPtr<UNNEModelData> ModelData) = 0;
 };
\ No newline at end of file
diff --git a/Engine/Plugins/Experimental/NNE/Source/NNE/Public/NNERuntimeRDG.h b/Engine/Plugins/Experimental/NNE/Source/NNE/Public/NNERuntimeRDG.h
--- a/Engine/Plugins/Experimental/NNE/Source/NNE/Public/NNERuntimeRDG.h
+++ b/Engine/Plugins/Experimental/NNE/Source/NNE/Public/NNERuntimeRDG.h
@@ -11,7 +11,7 @@
 
 class UNNEModelData;
 
-namespace UE::NNECore
+namespace UE::NNE
 {
 
 /**
@@ -27,8 +27,8 @@ struct NNE_API FTensorBindingRDG
 /**
  * The interface of a model instance that can run on RDG.
  *
- * Use UE::NNECore::IModelRDG::CreateModelInstance() to get a model instance.
- * Use UE::NNECore::GetRuntime<INNERuntimeRDG>(RuntimeName) to get a runtime capable of creating RDG models.
+ * Use UE::NNE::IModelRDG::CreateModelInstance() to get a model instance.
+ * Use UE::NNE::GetRuntime<INNERuntimeRDG>(RuntimeName) to get a runtime capable of creating RDG models.
  */
 class NNE_API IModelInstanceRDG
 {
@@ -100,7 +100,7 @@ public:
 /**
  * The interface of a model capable of creating model instance that can run on RDG.
  *
- * Use UE::NNECore::GetRuntime<INNERuntimeRDG>(RuntimeName) to get a runtime capable of creating RDG models.
+ * Use UE::NNE::GetRuntime<INNERuntimeRDG>(RuntimeName) to get a runtime capable of creating RDG models.
  */
 class NNE_API IModelRDG
 {
@@ -116,10 +116,10 @@ public:
 	 *
 	 * @return A caller owned model representing the neural network instance created.
 	 */
-	virtual TUniquePtr<UE::NNECore::IModelInstanceRDG> CreateModelInstance() = 0;
+	virtual TUniquePtr<UE::NNE::IModelInstanceRDG> CreateModelInstance() = 0;
 };
 
-} // UE::NNECore
+} // UE::NNE
 
 UINTERFACE()
 class NNE_API UNNERuntimeRDG : public UInterface
@@ -130,7 +130,7 @@ class NNE_API UNNERuntimeRDG : public UInterface
 /**
  * The interface of a neural network runtime capable of creating RDG models.
  *
- * Call UE::NNECore::GetRuntime<INNERuntimeRDG>(RuntimeName) to get a runtime implementing this interface.
+ * Call UE::NNE::GetRuntime<INNERuntimeRDG>(RuntimeName) to get a runtime implementing this interface.
  */
 class NNE_API INNERuntimeRDG
 {
@@ -156,5 +156,5 @@ public:
 	 * @param ModelData The model data for which to create a model.
 	 * @return A caller owned model representing the neural network created from ModelData.
 	 */
-	virtual TUniquePtr<UE::NNECore::IModelRDG> CreateModel(TObjectPtr<UNNEModelData> ModelData) = 0;
+	virtual TUniquePtr<UE::NNE::IModelRDG> CreateModel(TObjectPtr<UNNEModelData> ModelData) = 0;
 };
\ No newline at end of file
diff --git a/Engine/Plugins/Experimental/NNE/Source/NNE/Public/NNETypes.h b/Engine/Plugins/Experimental/NNE/Source/NNE/Public/NNETypes.h
--- a/Engine/Plugins/Experimental/NNE/Source/NNE/Public/NNETypes.h
+++ b/Engine/Plugins/Experimental/NNE/Source/NNE/Public/NNETypes.h
@@ -32,7 +32,7 @@ enum class ENNETensorDataType : uint8
 	BFloat16							//!< 16-bit floating number
 };
 
-namespace UE::NNECore
+namespace UE::NNE
 {
 
 	/**
@@ -277,4 +277,4 @@ namespace UE::NNECore
 		inline const FSymbolicTensorShape& GetShape() const { return Shape; }
 	};
 
-} // namespace UE::NNECore
+} // namespace UE::NNE
diff --git a/Engine/Plugins/Experimental/NNE/Source/NNEEditorTools/Private/SNNEEditorToolsRuntimesWidget.cpp b/Engine/Plugins/Experimental/NNE/Source/NNEEditorTools/Private/SNNEEditorToolsRuntimesWidget.cpp
--- a/Engine/Plugins/Experimental/NNE/Source/NNEEditorTools/Private/SNNEEditorToolsRuntimesWidget.cpp
+++ b/Engine/Plugins/Experimental/NNE/Source/NNEEditorTools/Private/SNNEEditorToolsRuntimesWidget.cpp
@@ -126,7 +126,7 @@ void SRuntimesWidget::InitCheckBoxNames()
 
 	CheckBoxNames.Add(TEXT("All"));
 
-	for (TWeakInterfacePtr<INNERuntime> Runtime : UE::NNECore::GetAllRuntimes())
+	for (TWeakInterfacePtr<INNERuntime> Runtime : UE::NNE::GetAllRuntimes())
 	{
 		CheckBoxNames.Add(Runtime->GetRuntimeName());
 	}
diff --git a/Engine/Plugins/Experimental/NNERuntimeORTCpu/Source/NNERuntimeORTCpu/Private/NNERuntimeORTCpu.cpp b/Engine/Plugins/Experimental/NNERuntimeORTCpu/Source/NNERuntimeORTCpu/Private/NNERuntimeORTCpu.cpp
--- a/Engine/Plugins/Experimental/NNERuntimeORTCpu/Source/NNERuntimeORTCpu/Private/NNERuntimeORTCpu.cpp
+++ b/Engine/Plugins/Experimental/NNERuntimeORTCpu/Source/NNERuntimeORTCpu/Private/NNERuntimeORTCpu.cpp
@@ -28,7 +28,7 @@ TArray<uint8> UNNERuntimeORTCpuImpl::CreateModelData(FString FileType, TConstArr
 		return {};
 	}
 
-	TUniquePtr<UE::NNECore::Internal::IModelOptimizer> Optimizer = UE::NNEUtils::Internal::CreateONNXToONNXModelOptimizer();
+	TUniquePtr<UE::NNE::Internal::IModelOptimizer> Optimizer = UE::NNEUtils::Internal::CreateONNXToONNXModelOptimizer();
 
 	FNNEModelRaw InputModel;
 	InputModel.Data = FileData;
@@ -68,18 +68,18 @@ bool UNNERuntimeORTCpuImpl::CanCreateModelCPU(TObjectPtr<UNNEModelData> ModelDat
 	return bResult;
 }
 
-TUniquePtr<UE::NNECore::IModelCPU> UNNERuntimeORTCpuImpl::CreateModel(TObjectPtr<UNNEModelData> ModelData)
+TUniquePtr<UE::NNE::IModelCPU> UNNERuntimeORTCpuImpl::CreateModel(TObjectPtr<UNNEModelData> ModelData)
 {
 	check(ModelData != nullptr);
 	
 	if (!CanCreateModelCPU(ModelData))
 	{
-		return TUniquePtr<UE::NNECore::IModelCPU>();
+		return TUniquePtr<UE::NNE::IModelCPU>();
 	}
 
 	TConstArrayView<uint8> Data = ModelData->GetModelData(GetRuntimeName());
 	UE::NNERuntimeORTCpu::Private::FModelCPU* Model = new UE::NNERuntimeORTCpu::Private::FModelCPU(&NNEEnvironmentCPU, Data);
-	UE::NNECore::IModelCPU* IModel = static_cast<UE::NNECore::IModelCPU*>(Model);
+	UE::NNE::IModelCPU* IModel = static_cast<UE::NNE::IModelCPU*>(Model);
 
 	if (FEngineAnalytics::IsAvailable())
 	{
@@ -91,5 +91,5 @@ TUniquePtr<UE::NNECore::IModelCPU> UNNERuntimeORTCpuImpl::CreateModel(TObjectPtr
 		FEngineAnalytics::GetProvider().RecordEvent(TEXT("NeuralNetworkEngine.CreateModel"), Attributes);
 	}
 
-	return TUniquePtr<UE::NNECore::IModelCPU>(IModel);
+	return TUniquePtr<UE::NNE::IModelCPU>(IModel);
 }
\ No newline at end of file
diff --git a/Engine/Plugins/Experimental/NNERuntimeORTCpu/Source/NNERuntimeORTCpu/Private/NNERuntimeORTCpu.h b/Engine/Plugins/Experimental/NNERuntimeORTCpu/Source/NNERuntimeORTCpu/Private/NNERuntimeORTCpu.h
--- a/Engine/Plugins/Experimental/NNERuntimeORTCpu/Source/NNERuntimeORTCpu/Private/NNERuntimeORTCpu.h
+++ b/Engine/Plugins/Experimental/NNERuntimeORTCpu/Source/NNERuntimeORTCpu/Private/NNERuntimeORTCpu.h
@@ -38,5 +38,5 @@ public:
 	virtual TArray<uint8> CreateModelData(FString FileType, TConstArrayView<uint8> FileData) override;
 
 	virtual bool CanCreateModelCPU(TObjectPtr<UNNEModelData> ModelData) const override;
-	virtual TUniquePtr<UE::NNECore::IModelCPU> CreateModel(TObjectPtr<UNNEModelData> ModelData) override;
+	virtual TUniquePtr<UE::NNE::IModelCPU> CreateModel(TObjectPtr<UNNEModelData> ModelData) override;
 };
\ No newline at end of file
diff --git a/Engine/Plugins/Experimental/NNERuntimeORTCpu/Source/NNERuntimeORTCpu/Private/NNERuntimeORTCpuModel.cpp b/Engine/Plugins/Experimental/NNERuntimeORTCpu/Source/NNERuntimeORTCpu/Private/NNERuntimeORTCpuModel.cpp
--- a/Engine/Plugins/Experimental/NNERuntimeORTCpu/Source/NNERuntimeORTCpu/Private/NNERuntimeORTCpuModel.cpp
+++ b/Engine/Plugins/Experimental/NNERuntimeORTCpu/Source/NNERuntimeORTCpu/Private/NNERuntimeORTCpuModel.cpp
@@ -121,7 +121,7 @@ namespace UE::NNERuntimeORTCpu::Private
 
 		const uint32 NumberTensors = bIsInput ? Session->GetInputCount() : Session->GetOutputCount();
 
-		TArray<NNECore::FTensorDesc>& SymbolicTensorDescs = bIsInput ? InputSymbolicTensors : OutputSymbolicTensors;
+		TArray<NNE::FTensorDesc>& SymbolicTensorDescs = bIsInput ? InputSymbolicTensors : OutputSymbolicTensors;
 		TArray<ONNXTensorElementDataType>& TensorsORTType = bIsInput ? InputTensorsORTType : OutputTensorsORTType;
 		TArray<Ort::AllocatedStringPtr>& TensorNames = bIsInput ? InputTensorNames : OutputTensorNames;
 		TArray<const char*>& TensorNamePtrs = bIsInput ? InputTensorNamePtrs : OutputTensorNamePtrs;
@@ -152,8 +152,8 @@ namespace UE::NNERuntimeORTCpu::Private
 				ShapeData.Add((int32)CurrentTensorSize);
 			}
 
-			NNECore::FSymbolicTensorShape Shape = NNECore::FSymbolicTensorShape::Make(ShapeData);
-			NNECore::FTensorDesc SymbolicTensorDesc = NNECore::FTensorDesc::Make(FString(CurTensorNamePtr), Shape, TypeAndSize.first);
+			NNE::FSymbolicTensorShape Shape = NNE::FSymbolicTensorShape::Make(ShapeData);
+			NNE::FTensorDesc SymbolicTensorDesc = NNE::FTensorDesc::Make(FString(CurTensorNamePtr), Shape, TypeAndSize.first);
 
 			check(SymbolicTensorDesc.GetElemByteSize() == TypeAndSize.second);
 			SymbolicTensorDescs.Emplace(SymbolicTensorDesc);
@@ -162,7 +162,7 @@ namespace UE::NNERuntimeORTCpu::Private
 		return true;
 	}
 	
-	int FModelInstanceCPU::SetInputTensorShapes(TConstArrayView<NNECore::FTensorShape> InInputShapes)
+	int FModelInstanceCPU::SetInputTensorShapes(TConstArrayView<NNE::FTensorShape> InInputShapes)
 	{
 		InputTensors.Reset();
 		OutputTensors.Reset();
@@ -177,7 +177,7 @@ namespace UE::NNERuntimeORTCpu::Private
 		// Setup concrete input tensor
 		for (int i = 0; i < InputSymbolicTensors.Num(); ++i)
 		{
-			NNECore::Internal::FTensor Tensor = NNECore::Internal::FTensor::Make(InputSymbolicTensors[i].GetName(), InInputShapes[i], InputSymbolicTensors[i].GetDataType());
+			NNE::Internal::FTensor Tensor = NNE::Internal::FTensor::Make(InputSymbolicTensors[i].GetName(), InInputShapes[i], InputSymbolicTensors[i].GetDataType());
 			InputTensors.Emplace(Tensor);
 		}
 
@@ -185,11 +185,11 @@ namespace UE::NNERuntimeORTCpu::Private
 		// this would allow to resolve output shapes here rather than during inference.
 
 		// Setup concrete output shapes only if all model output shapes are concretes, otherwise it will be set during Run()
-		for (NNECore::FTensorDesc SymbolicTensorDesc : OutputSymbolicTensors)
+		for (NNE::FTensorDesc SymbolicTensorDesc : OutputSymbolicTensors)
 		{
 			if (SymbolicTensorDesc.GetShape().IsConcrete())
 			{
-				NNECore::Internal::FTensor Tensor = NNECore::Internal::FTensor::MakeFromSymbolicDesc(SymbolicTensorDesc);
+				NNE::Internal::FTensor Tensor = NNE::Internal::FTensor::MakeFromSymbolicDesc(SymbolicTensorDesc);
 				OutputTensors.Emplace(Tensor);
 				OutputTensorShapes.Emplace(Tensor.GetShape());
 			}
@@ -203,7 +203,7 @@ namespace UE::NNERuntimeORTCpu::Private
 		return 0;
 	}
 
-	int FModelInstanceCPU::RunSync(TConstArrayView<NNECore::FTensorBindingCPU> InInputBindings, TConstArrayView<NNECore::FTensorBindingCPU> InOutputBindings)
+	int FModelInstanceCPU::RunSync(TConstArrayView<NNE::FTensorBindingCPU> InInputBindings, TConstArrayView<NNE::FTensorBindingCPU> InOutputBindings)
 	{
 		DECLARE_SCOPE_CYCLE_COUNTER(TEXT("FModelInstanceCPU_Run"), STAT_FModelInstanceCPU_Run, STATGROUP_NNE);
 
@@ -300,7 +300,7 @@ namespace UE::NNERuntimeORTCpu::Private
 		InputTransferStatisticsEstimator.ResetStats();
 	}
 
-	TUniquePtr<UE::NNECore::IModelInstanceCPU> FModelCPU::CreateModelInstance()
+	TUniquePtr<UE::NNE::IModelInstanceCPU> FModelCPU::CreateModelInstance()
 	{
 		const UE::NNERuntimeORTCpu::Private::FRuntimeConf InConf;
 		UE::NNERuntimeORTCpu::Private::FModelInstanceCPU* ModelInstance = new UE::NNERuntimeORTCpu::Private::FModelInstanceCPU(ORTEnvironment, InConf);
@@ -308,11 +308,11 @@ namespace UE::NNERuntimeORTCpu::Private
 		if (!ModelInstance->Init(ModelData))
 		{
 			delete ModelInstance;
-			return TUniquePtr<UE::NNECore::IModelInstanceCPU>();
+			return TUniquePtr<UE::NNE::IModelInstanceCPU>();
 		}
 
-		UE::NNECore::IModelInstanceCPU* IModelInstance = static_cast<UE::NNECore::IModelInstanceCPU*>(ModelInstance);
-		return TUniquePtr<UE::NNECore::IModelInstanceCPU>(IModelInstance);
+		UE::NNE::IModelInstanceCPU* IModelInstance = static_cast<UE::NNE::IModelInstanceCPU*>(ModelInstance);
+		return TUniquePtr<UE::NNE::IModelInstanceCPU>(IModelInstance);
 	}
 
 	FModelCPU::FModelCPU(Ort::Env* InORTEnvironment, TConstArrayView<uint8> InModelData) :
diff --git a/Engine/Plugins/Experimental/NNERuntimeORTCpu/Source/NNERuntimeORTCpu/Private/NNERuntimeORTCpuModel.h b/Engine/Plugins/Experimental/NNERuntimeORTCpu/Source/NNERuntimeORTCpu/Private/NNERuntimeORTCpuModel.h
--- a/Engine/Plugins/Experimental/NNERuntimeORTCpu/Source/NNERuntimeORTCpu/Private/NNERuntimeORTCpuModel.h
+++ b/Engine/Plugins/Experimental/NNERuntimeORTCpu/Source/NNERuntimeORTCpu/Private/NNERuntimeORTCpuModel.h
@@ -25,7 +25,7 @@ namespace UE::NNERuntimeORTCpu::Private
 		EThreadPriority ThreadPriority = EThreadPriority::TPri_Normal;
 	};
 
-	class FModelInstanceCPU : public NNECore::Internal::FModelInstanceBase<NNECore::IModelInstanceCPU>
+	class FModelInstanceCPU : public NNE::Internal::FModelInstanceBase<NNE::IModelInstanceCPU>
 	{
 
 	public:
@@ -33,8 +33,8 @@ namespace UE::NNERuntimeORTCpu::Private
 		FModelInstanceCPU(Ort::Env* InORTEnvironment, const FRuntimeConf& InRuntimeConf);
 		virtual ~FModelInstanceCPU() {};
 
-		virtual int SetInputTensorShapes(TConstArrayView<NNECore::FTensorShape> InInputShapes) override;
-		virtual int RunSync(TConstArrayView<NNECore::FTensorBindingCPU> InInputBindings, TConstArrayView<NNECore::FTensorBindingCPU> InOutputBindings) override;
+		virtual int SetInputTensorShapes(TConstArrayView<NNE::FTensorShape> InInputShapes) override;
+		virtual int RunSync(TConstArrayView<NNE::FTensorBindingCPU> InInputBindings, TConstArrayView<NNE::FTensorBindingCPU> InOutputBindings) override;
 
 		bool Init(TConstArrayView<uint8> ModelData);
 		bool IsLoaded() const;
@@ -66,8 +66,8 @@ namespace UE::NNERuntimeORTCpu::Private
 		TArray<const char*> InputTensorNamePtrs;
 		TArray<const char*> OutputTensorNamePtrs;
 
-		TArray<NNECore::Internal::FTensor> InputTensors;
-		TArray<NNECore::Internal::FTensor> OutputTensors;
+		TArray<NNE::Internal::FTensor> InputTensors;
+		TArray<NNE::Internal::FTensor> OutputTensors;
 
 		bool InitializedAndConfigureMembers();
 		bool ConfigureTensors(const bool InIsInput);
@@ -80,13 +80,13 @@ namespace UE::NNERuntimeORTCpu::Private
 
 	};
 
-	class FModelCPU : public NNECore::IModelCPU
+	class FModelCPU : public NNE::IModelCPU
 	{
 	public:
 		FModelCPU(Ort::Env* InORTEnvironment, TConstArrayView<uint8> ModelData);
 		virtual ~FModelCPU() {};
 
-		virtual TUniquePtr<UE::NNECore::IModelInstanceCPU> CreateModelInstance() override;
+		virtual TUniquePtr<UE::NNE::IModelInstanceCPU> CreateModelInstance() override;
 
 	private:
 		Ort::Env* ORTEnvironment;
diff --git a/Engine/Plugins/Experimental/NNERuntimeORTCpu/Source/NNERuntimeORTCpu/Private/NNERuntimeORTCpuModule.cpp b/Engine/Plugins/Experimental/NNERuntimeORTCpu/Source/NNERuntimeORTCpu/Private/NNERuntimeORTCpuModule.cpp
--- a/Engine/Plugins/Experimental/NNERuntimeORTCpu/Source/NNERuntimeORTCpu/Private/NNERuntimeORTCpuModule.cpp
+++ b/Engine/Plugins/Experimental/NNERuntimeORTCpu/Source/NNERuntimeORTCpu/Private/NNERuntimeORTCpuModule.cpp
@@ -14,7 +14,7 @@ void FNNERuntimeORTCpuModule::StartupModule()
 		TWeakInterfacePtr<INNERuntime> RuntimeCPUInterface(NNERuntimeORTCpu.Get());
 		
 		NNERuntimeORTCpu->AddToRoot();
-		UE::NNECore::RegisterRuntime(RuntimeCPUInterface);
+		UE::NNE::RegisterRuntime(RuntimeCPUInterface);
 	}
 }
 
@@ -25,7 +25,7 @@ void FNNERuntimeORTCpuModule::ShutdownModule()
 	{
 		TWeakInterfacePtr<INNERuntime> RuntimeCPUInterface(NNERuntimeORTCpu.Get());
 		
-		UE::NNECore::UnregisterRuntime(RuntimeCPUInterface);
+		UE::NNE::UnregisterRuntime(RuntimeCPUInterface);
 		NNERuntimeORTCpu->RemoveFromRoot();
 		NNERuntimeORTCpu = TWeakObjectPtr<UNNERuntimeORTCpuImpl>(nullptr);
 	}
diff --git a/Engine/Plugins/Experimental/NNERuntimeORTCpu/Source/NNERuntimeORTCpu/Private/NNERuntimeORTCpuUtils.h b/Engine/Plugins/Experimental/NNERuntimeORTCpu/Source/NNERuntimeORTCpu/Private/NNERuntimeORTCpuUtils.h
--- a/Engine/Plugins/Experimental/NNERuntimeORTCpu/Source/NNERuntimeORTCpu/Private/NNERuntimeORTCpuUtils.h
+++ b/Engine/Plugins/Experimental/NNERuntimeORTCpu/Source/NNERuntimeORTCpu/Private/NNERuntimeORTCpuUtils.h
@@ -118,8 +118,8 @@ namespace UE::NNERuntimeORTCpu::Private
 	}
 
 	inline void BindTensorsToORT(
-		TConstArrayView<NNECore::FTensorBindingCPU> InBindingTensors,
-		TConstArrayView<NNECore::Internal::FTensor> InTensors,
+		TConstArrayView<NNE::FTensorBindingCPU> InBindingTensors,
+		TConstArrayView<NNE::Internal::FTensor> InTensors,
 		TConstArrayView<ONNXTensorElementDataType> InTensorsORTType,
 		const Ort::MemoryInfo* InAllocatorInfo,
 		TArray<Ort::Value>& OutOrtTensors
@@ -136,8 +136,8 @@ namespace UE::NNERuntimeORTCpu::Private
 
 		for (uint32 Index = 0; Index < NumBinding; ++Index)
 		{
-			const NNECore::FTensorBindingCPU& Binding = InBindingTensors[Index];
-			const NNECore::Internal::FTensor& Tensor = InTensors[Index];
+			const NNE::FTensorBindingCPU& Binding = InBindingTensors[Index];
+			const NNE::Internal::FTensor& Tensor = InTensors[Index];
 			const ONNXTensorElementDataType CurrentORTType = InTensorsORTType[Index];
 
 			TUniquePtr<int64_t[]> SizesInt64t;
@@ -164,9 +164,9 @@ namespace UE::NNERuntimeORTCpu::Private
 
 	inline void CopyFromORTToBindings(
 		TConstArrayView<Ort::Value> InOrtTensors,
-		TConstArrayView<NNECore::FTensorBindingCPU> InBindingTensors,
-		TConstArrayView<NNECore::FTensorDesc> InTensorDescs,
-		TArray<NNECore::Internal::FTensor>& OutTensors
+		TConstArrayView<NNE::FTensorBindingCPU> InBindingTensors,
+		TConstArrayView<NNE::FTensorDesc> InTensorDescs,
+		TArray<NNE::Internal::FTensor>& OutTensors
 	)
 	{
 		const uint32 NumBinding = InOrtTensors.Num();
@@ -180,8 +180,8 @@ namespace UE::NNERuntimeORTCpu::Private
 
 		for (uint32 Index = 0; Index < NumBinding; ++Index)
 		{
-			const NNECore::FTensorBindingCPU& Binding = InBindingTensors[Index];
-			const NNECore::FTensorDesc& TensorDesc = InTensorDescs[Index];
+			const NNE::FTensorBindingCPU& Binding = InBindingTensors[Index];
+			const NNE::FTensorDesc& TensorDesc = InTensorDescs[Index];
 			const Ort::Value& OrtTensor = InOrtTensors[Index];
 			const std::vector<int64_t>& OrtShape = OrtTensor.GetTensorTypeAndShapeInfo().GetShape();
 
@@ -192,8 +192,8 @@ namespace UE::NNERuntimeORTCpu::Private
 				ShapeData.Add(OrtShape[DimIndex]);
 			}
 
-			NNECore::FTensorShape Shape = NNECore::FTensorShape::Make(ShapeData);
-			NNECore::Internal::FTensor Tensor = NNECore::Internal::FTensor::Make(TensorDesc.GetName(), Shape, TensorDesc.GetDataType());
+			NNE::FTensorShape Shape = NNE::FTensorShape::Make(ShapeData);
+			NNE::Internal::FTensor Tensor = NNE::Internal::FTensor::Make(TensorDesc.GetName(), Shape, TensorDesc.GetDataType());
 			OutTensors.Emplace(Tensor);
 
 			const void* OrtTensorData = OrtTensor.GetTensorData<void>();
diff --git a/Engine/Plugins/Experimental/NNERuntimeORTCpu/Source/NNEUtils/Internal/NNEUtilsModelBuilder.h b/Engine/Plugins/Experimental/NNERuntimeORTCpu/Source/NNEUtils/Internal/NNEUtilsModelBuilder.h
--- a/Engine/Plugins/Experimental/NNERuntimeORTCpu/Source/NNEUtils/Internal/NNEUtilsModelBuilder.h
+++ b/Engine/Plugins/Experimental/NNERuntimeORTCpu/Source/NNEUtils/Internal/NNEUtilsModelBuilder.h
@@ -4,13 +4,13 @@
 #include "NNETypes.h"
 #include "NNETensor.h"
 
-namespace UE::NNECore { class FAttributeMap; }
+namespace UE::NNE { class FAttributeMap; }
 struct FNNEAttributeValue;
 struct FNNEModelRaw;
 
 namespace UE::NNEUtils::Internal
 {
-using FTensor = NNECore::Internal::FTensor;
+using FTensor = NNE::Internal::FTensor;
 
 class IModelBuilder
 {
@@ -92,7 +92,7 @@ NNEUTILS_API IModelBuilder* CreateONNXModelBuilder(int64 IrVersion = OnnxIrVersi
 NNEUTILS_API bool CreateONNXModelForOperator(bool UseVariadicShapeForModel, const FString& OperatorName, 
 	TConstArrayView<FTensor> InInputTensors, TConstArrayView<FTensor> InOutputTensors,
 	TConstArrayView<FTensor> InWeightTensors, TConstArrayView<TArray<char>> InWeightTensorsData,
-	const UE::NNECore::FAttributeMap& Attributes, FNNEModelRaw& ModelData);
+	const UE::NNE::FAttributeMap& Attributes, FNNEModelRaw& ModelData);
 
 /**
  * Create an instance of NNE model builder that creates NNE model/format in memory
diff --git a/Engine/Plugins/Experimental/NNERuntimeORTCpu/Source/NNEUtils/Internal/NNEUtilsModelOptimizer.h b/Engine/Plugins/Experimental/NNERuntimeORTCpu/Source/NNEUtils/Internal/NNEUtilsModelOptimizer.h
--- a/Engine/Plugins/Experimental/NNERuntimeORTCpu/Source/NNEUtils/Internal/NNEUtilsModelOptimizer.h
+++ b/Engine/Plugins/Experimental/NNERuntimeORTCpu/Source/NNEUtils/Internal/NNEUtilsModelOptimizer.h
@@ -3,28 +3,28 @@
 
 #include "NNERuntimeFormat.h"
 
-namespace UE::NNECore { class FAttributeMap; }
-namespace UE::NNECore::Internal { class IModelOptimizer; }
+namespace UE::NNE { class FAttributeMap; }
+namespace UE::NNE::Internal { class IModelOptimizer; }
 
 namespace UE::NNEUtils::Internal
 {
 
-using FOptimizerOptionsMap = NNECore::FAttributeMap;
+using FOptimizerOptionsMap = NNE::FAttributeMap;
 	
 /** Create a model optimizer */
-NNEUTILS_API TUniquePtr<NNECore::Internal::IModelOptimizer> CreateModelOptimizer(ENNEInferenceFormat InputFormat, ENNEInferenceFormat OutputFormat);
+NNEUTILS_API TUniquePtr<NNE::Internal::IModelOptimizer> CreateModelOptimizer(ENNEInferenceFormat InputFormat, ENNEInferenceFormat OutputFormat);
 
-inline TUniquePtr<NNECore::Internal::IModelOptimizer> CreateONNXToNNEModelOptimizer()
+inline TUniquePtr<NNE::Internal::IModelOptimizer> CreateONNXToNNEModelOptimizer()
 {
 	return CreateModelOptimizer(ENNEInferenceFormat::ONNX, ENNEInferenceFormat::NNERT);
 }
 
-inline TUniquePtr<NNECore::Internal::IModelOptimizer> CreateONNXToORTModelOptimizer()
+inline TUniquePtr<NNE::Internal::IModelOptimizer> CreateONNXToORTModelOptimizer()
 {
 	return CreateModelOptimizer(ENNEInferenceFormat::ONNX, ENNEInferenceFormat::ORT);
 }
 
-inline TUniquePtr<NNECore::Internal::IModelOptimizer> CreateONNXToONNXModelOptimizer()
+inline TUniquePtr<NNE::Internal::IModelOptimizer> CreateONNXToONNXModelOptimizer()
 {
 	return CreateModelOptimizer(ENNEInferenceFormat::ONNX, ENNEInferenceFormat::ONNX);
 }
diff --git a/Engine/Plugins/Experimental/NNERuntimeORTCpu/Source/NNEUtils/Private/NNEUtilsModelBuilderNNE.cpp b/Engine/Plugins/Experimental/NNERuntimeORTCpu/Source/NNEUtils/Private/NNEUtilsModelBuilderNNE.cpp
--- a/Engine/Plugins/Experimental/NNERuntimeORTCpu/Source/NNEUtils/Private/NNEUtilsModelBuilderNNE.cpp
+++ b/Engine/Plugins/Experimental/NNERuntimeORTCpu/Source/NNEUtils/Private/NNEUtilsModelBuilderNNE.cpp
@@ -110,7 +110,7 @@ public:
 
 	virtual HTensor AddTensor(const FString& Name, ENNETensorDataType DataType, TArrayView<const int32> Shape, const void* Data, uint64 DataSize)
 	{
-		TArray<int32, TInlineAllocator<NNECore::FTensorShape::MaxRank>> NNEShape;
+		TArray<int32, TInlineAllocator<NNE::FTensorShape::MaxRank>> NNEShape;
 		for (int i = 0; i < Shape.Num(); ++i)
 		{
 			//ORT Graph return 0 for variable dimensions for inputs/outputs, NNE use -1.
diff --git a/Engine/Plugins/Experimental/NNERuntimeORTCpu/Source/NNEUtils/Private/NNEUtilsModelBuilderONNX.cpp b/Engine/Plugins/Experimental/NNERuntimeORTCpu/Source/NNEUtils/Private/NNEUtilsModelBuilderONNX.cpp
--- a/Engine/Plugins/Experimental/NNERuntimeORTCpu/Source/NNEUtils/Private/NNEUtilsModelBuilderONNX.cpp
+++ b/Engine/Plugins/Experimental/NNERuntimeORTCpu/Source/NNEUtils/Private/NNEUtilsModelBuilderONNX.cpp
@@ -324,7 +324,7 @@ private:
 //
 //
 //
-void BuildShapeForModel(bool ConvertToVariadicShape, const NNECore::FTensorShape& InShape, TArray<int32>& OutShape)
+void BuildShapeForModel(bool ConvertToVariadicShape, const NNE::FTensorShape& InShape, TArray<int32>& OutShape)
 {
 	OutShape.Reset();
 	for (int32 Idx = 0; Idx < InShape.Rank(); ++Idx)
@@ -344,7 +344,7 @@ void BuildShapeForModel(bool ConvertToVariadicShape, const NNECore::FTensorShape
 NNEUTILS_API bool CreateONNXModelForOperator(bool UseVariadicShapeForModel, const FString& OperatorName,
 	TConstArrayView<FTensor> InInputTensors, TConstArrayView<FTensor> InOutputTensors,
 	TConstArrayView<FTensor> InWeightTensors, TConstArrayView<TArray<char>> InWeightTensorsData,
-	const UE::NNECore::FAttributeMap& Attributes, FNNEModelRaw& Model)
+	const UE::NNE::FAttributeMap& Attributes, FNNEModelRaw& Model)
 {
 	Model = FNNEModelRaw{};
 	
diff --git a/Engine/Plugins/Experimental/NNERuntimeORTCpu/Source/NNEUtils/Private/NNEUtilsModelOptimizer.cpp b/Engine/Plugins/Experimental/NNERuntimeORTCpu/Source/NNEUtils/Private/NNEUtilsModelOptimizer.cpp
--- a/Engine/Plugins/Experimental/NNERuntimeORTCpu/Source/NNEUtils/Private/NNEUtilsModelOptimizer.cpp
+++ b/Engine/Plugins/Experimental/NNERuntimeORTCpu/Source/NNEUtils/Private/NNEUtilsModelOptimizer.cpp
@@ -265,11 +265,11 @@ ENNETensorDataType GetDataTypeFromGraphTensor(Ort::GraphTensorDataType TensorDat
 
 #undef CASE
 
-class FModelOptimizerBase : public NNECore::Internal::IModelOptimizer
+class FModelOptimizerBase : public NNE::Internal::IModelOptimizer
 {
 protected:
-	TArray<TSharedPtr<NNECore::Internal::IModelOptimizerPass>> OptimizationPasses;
-	TArray<TSharedPtr<NNECore::Internal::IModelValidator>> Validators;
+	TArray<TSharedPtr<NNE::Internal::IModelOptimizerPass>> OptimizationPasses;
+	TArray<TSharedPtr<NNE::Internal::IModelValidator>> Validators;
 
 	virtual bool ValidateInputModel(const FNNEModelRaw& InputModel)
 	{
@@ -290,7 +290,7 @@ protected:
 	}
 	
 public:
-	virtual void AddOptimizationPass(TSharedPtr<NNECore::Internal::IModelOptimizerPass> ModelOptimizerPass) override
+	virtual void AddOptimizationPass(TSharedPtr<NNE::Internal::IModelOptimizerPass> ModelOptimizerPass) override
 	{
 		if (ModelOptimizerPass.IsValid())
 		{
@@ -298,7 +298,7 @@ public:
 		}
 	}
 
-	virtual void AddValidator(TSharedPtr<NNECore::Internal::IModelValidator> ModelValidator) override
+	virtual void AddValidator(TSharedPtr<NNE::Internal::IModelValidator> ModelValidator) override
 	{
 		if (ModelValidator.IsValid())
 		{
@@ -310,7 +310,7 @@ public:
 	{
 		bool bIsModelValid = true;
 
-		for (TSharedPtr<NNECore::Internal::IModelValidator>& Validator : Validators)
+		for (TSharedPtr<NNE::Internal::IModelValidator>& Validator : Validators)
 		{
 			check(Validator.IsValid());
 			if (!Validator->ValidateModel(ModelToValidate, Options))
@@ -330,7 +330,7 @@ public:
 			return false;
 		}
 		
-		for (TSharedPtr<NNECore::Internal::IModelOptimizerPass>& Pass : OptimizationPasses)
+		for (TSharedPtr<NNE::Internal::IModelOptimizerPass>& Pass : OptimizationPasses)
 		{
 			check(Pass.IsValid());
 			
@@ -363,7 +363,7 @@ public:
 	}
 };
 
-class FONNXModelValidator : public NNECore::Internal::IModelValidator
+class FONNXModelValidator : public NNE::Internal::IModelValidator
 {
 public:
 	virtual FString GetName() const
@@ -699,7 +699,7 @@ private:
 };
 
 /** Create a model optimizer */
-NNEUTILS_API TUniquePtr<NNECore::Internal::IModelOptimizer> CreateModelOptimizer(ENNEInferenceFormat InputFormat, ENNEInferenceFormat OutputFormat)
+NNEUTILS_API TUniquePtr<NNE::Internal::IModelOptimizer> CreateModelOptimizer(ENNEInferenceFormat InputFormat, ENNEInferenceFormat OutputFormat)
 {
 	if (InputFormat == ENNEInferenceFormat::ONNX)
 	{
diff --git a/Engine/Plugins/Experimental/NNERuntimeORTGpu/Source/NNERuntimeORT/Private/NNERuntimeORT.cpp b/Engine/Plugins/Experimental/NNERuntimeORTGpu/Source/NNERuntimeORT/Private/NNERuntimeORT.cpp
--- a/Engine/Plugins/Experimental/NNERuntimeORTGpu/Source/NNERuntimeORT/Private/NNERuntimeORT.cpp
+++ b/Engine/Plugins/Experimental/NNERuntimeORTGpu/Source/NNERuntimeORT/Private/NNERuntimeORT.cpp
@@ -30,13 +30,13 @@ TArray<uint8> UNNERuntimeORTGpuImpl::CreateModelData(FString FileType, TConstArr
 		return {};
 	}
 
-	TUniquePtr<UE::NNECore::Internal::IModelOptimizer> Optimizer = UE::NNEUtils::Internal::CreateONNXToONNXModelOptimizer();
+	TUniquePtr<UE::NNE::Internal::IModelOptimizer> Optimizer = UE::NNEUtils::Internal::CreateONNXToONNXModelOptimizer();
 
 	FNNEModelRaw InputModel;
 	InputModel.Data = FileData;
 	InputModel.Format = ENNEInferenceFormat::ONNX;
 	FNNEModelRaw OutputModel;
-	UE::NNECore::Internal::FOptimizerOptionsMap Options;
+	UE::NNE::Internal::FOptimizerOptionsMap Options;
 	if (!Optimizer->Optimize(InputModel, OutputModel, Options))
 	{
 		return {};
@@ -93,30 +93,30 @@ bool UNNERuntimeORTGpuImpl::CanCreateModelGPU(TObjectPtr<UNNEModelData> ModelDat
 	return bResult;
 }
 
-TUniquePtr<UE::NNECore::IModelGPU> UNNERuntimeORTGpuImpl::CreateModel(TObjectPtr<UNNEModelData> ModelData)
+TUniquePtr<UE::NNE::IModelGPU> UNNERuntimeORTGpuImpl::CreateModel(TObjectPtr<UNNEModelData> ModelData)
 {
 	check(ModelData != nullptr);
 	check(ORTEnvironment.IsValid());
 
 	if (!CanCreateModelGPU(ModelData))
 	{
-		return TUniquePtr<UE::NNECore::IModelGPU>();
+		return TUniquePtr<UE::NNE::IModelGPU>();
 	}
 
-	UE::NNECore::IModelGPU* IModel = nullptr;
+	UE::NNE::IModelGPU* IModel = nullptr;
 	TConstArrayView<uint8> Data = ModelData->GetModelData(GetRuntimeName());
 
 	switch (Provider)
 	{
 		case ENNERuntimeORTGpuProvider::Dml:  
-			IModel = static_cast<UE::NNECore::IModelGPU*>(new UE::NNERuntimeORT::Private::FModelORTDml(ORTEnvironment.Get(), Data));
+			IModel = static_cast<UE::NNE::IModelGPU*>(new UE::NNERuntimeORT::Private::FModelORTDml(ORTEnvironment.Get(), Data));
 			break;
 		case ENNERuntimeORTGpuProvider::Cuda: 
-			IModel = static_cast<UE::NNECore::IModelGPU*>(new UE::NNERuntimeORT::Private::FModelORTCuda(ORTEnvironment.Get(), Data));
+			IModel = static_cast<UE::NNE::IModelGPU*>(new UE::NNERuntimeORT::Private::FModelORTCuda(ORTEnvironment.Get(), Data));
 			break;
 		default:
 			UE_LOG(LogNNE, Error, TEXT("Failed to create model for ORT GPU runtime, unsupported provider. Runtime will not be functional."));
-			return TUniquePtr<UE::NNECore::IModelGPU>();
+			return TUniquePtr<UE::NNE::IModelGPU>();
 	}
 
 	if (FEngineAnalytics::IsAvailable())
@@ -129,7 +129,7 @@ TUniquePtr<UE::NNECore::IModelGPU> UNNERuntimeORTGpuImpl::CreateModel(TObjectPtr
 		FEngineAnalytics::GetProvider().RecordEvent(TEXT("NeuralNetworkEngine.CreateModel"), Attributes);
 	}
 
-	return TUniquePtr<UE::NNECore::IModelGPU>(IModel);
+	return TUniquePtr<UE::NNE::IModelGPU>(IModel);
 }
 
 #else // PLATFORM_WINDOWS
@@ -139,9 +139,9 @@ bool UNNERuntimeORTGpuImpl::CanCreateModelGPU(TObjectPtr<UNNEModelData> ModelDat
 	return false;
 }
 
-TUniquePtr<UE::NNECore::IModelGPU> UNNERuntimeORTGpuImpl::CreateModelGPU(TObjectPtr<UNNEModelData> ModelData)
+TUniquePtr<UE::NNE::IModelGPU> UNNERuntimeORTGpuImpl::CreateModelGPU(TObjectPtr<UNNEModelData> ModelData)
 {
-	return TUniquePtr<UE::NNECore::IModelGPU>();
+	return TUniquePtr<UE::NNE::IModelGPU>();
 }
 
 #endif // PLATFORM_WINDOWS
\ No newline at end of file
diff --git a/Engine/Plugins/Experimental/NNERuntimeORTGpu/Source/NNERuntimeORT/Private/NNERuntimeORT.h b/Engine/Plugins/Experimental/NNERuntimeORTGpu/Source/NNERuntimeORT/Private/NNERuntimeORT.h
--- a/Engine/Plugins/Experimental/NNERuntimeORTGpu/Source/NNERuntimeORT/Private/NNERuntimeORT.h
+++ b/Engine/Plugins/Experimental/NNERuntimeORTGpu/Source/NNERuntimeORT/Private/NNERuntimeORT.h
@@ -50,5 +50,5 @@ public:
 	virtual TArray<uint8> CreateModelData(FString FileType, TConstArrayView<uint8> FileData) override;
 
 	virtual bool CanCreateModelGPU(TObjectPtr<UNNEModelData> ModelData) const override;
-	virtual TUniquePtr<UE::NNECore::IModelGPU> CreateModel(TObjectPtr<UNNEModelData> ModelData) override;
+	virtual TUniquePtr<UE::NNE::IModelGPU> CreateModel(TObjectPtr<UNNEModelData> ModelData) override;
 };
\ No newline at end of file
diff --git a/Engine/Plugins/Experimental/NNERuntimeORTGpu/Source/NNERuntimeORT/Private/NNERuntimeORTModel.cpp b/Engine/Plugins/Experimental/NNERuntimeORTGpu/Source/NNERuntimeORT/Private/NNERuntimeORTModel.cpp
--- a/Engine/Plugins/Experimental/NNERuntimeORTGpu/Source/NNERuntimeORT/Private/NNERuntimeORTModel.cpp
+++ b/Engine/Plugins/Experimental/NNERuntimeORTGpu/Source/NNERuntimeORT/Private/NNERuntimeORTModel.cpp
@@ -149,7 +149,7 @@ namespace UE::NNERuntimeORT::Private
 
 		const uint32 NumberTensors = bIsInput ? Session->GetInputCount() : Session->GetOutputCount();
 
-		TArray<NNECore::FTensorDesc>& SymbolicTensorDescs = bIsInput ? InputSymbolicTensors : OutputSymbolicTensors;
+		TArray<NNE::FTensorDesc>& SymbolicTensorDescs = bIsInput ? InputSymbolicTensors : OutputSymbolicTensors;
 		TArray<ONNXTensorElementDataType>& TensorsORTType = bIsInput ? InputTensorsORTType : OutputTensorsORTType;
 		TArray<char*>& TensorNames = bIsInput ? InputTensorNames : OutputTensorNames;
 		TArray<Ort::AllocatedStringPtr>& TensorNameValues = bIsInput ? InputTensorNameValues : OutputTensorNameValues;
@@ -178,8 +178,8 @@ namespace UE::NNERuntimeORT::Private
 				ShapeData.Add((int32)CurrentTensorSize);
 			}
 
-			NNECore::FSymbolicTensorShape Shape = NNECore::FSymbolicTensorShape::Make(ShapeData);
-			NNECore::FTensorDesc SymbolicTensorDesc = NNECore::FTensorDesc::Make(FString(TensorNames.Last()), Shape, TypeAndSize.first);
+			NNE::FSymbolicTensorShape Shape = NNE::FSymbolicTensorShape::Make(ShapeData);
+			NNE::FTensorDesc SymbolicTensorDesc = NNE::FTensorDesc::Make(FString(TensorNames.Last()), Shape, TypeAndSize.first);
 
 			check(SymbolicTensorDesc.GetElemByteSize() == TypeAndSize.second);
 			SymbolicTensorDescs.Emplace(SymbolicTensorDesc);
@@ -188,7 +188,7 @@ namespace UE::NNERuntimeORT::Private
 		return true;
 	}
 	
-	int FModelInstanceORT::SetInputTensorShapes(TConstArrayView<NNECore::FTensorShape> InInputShapes)
+	int FModelInstanceORT::SetInputTensorShapes(TConstArrayView<NNE::FTensorShape> InInputShapes)
 	{
 		InputTensors.Reset();
 		OutputTensors.Reset();
@@ -203,7 +203,7 @@ namespace UE::NNERuntimeORT::Private
 		// Setup concrete input tensor
 		for (int i = 0; i < InputSymbolicTensors.Num(); ++i)
 		{
-			NNECore::Internal::FTensor Tensor = NNECore::Internal::FTensor::Make(InputSymbolicTensors[i].GetName(), InInputShapes[i], InputSymbolicTensors[i].GetDataType());
+			NNE::Internal::FTensor Tensor = NNE::Internal::FTensor::Make(InputSymbolicTensors[i].GetName(), InInputShapes[i], InputSymbolicTensors[i].GetDataType());
 			InputTensors.Emplace(Tensor);
 		}
 
@@ -211,11 +211,11 @@ namespace UE::NNERuntimeORT::Private
 		// this would allow to resolve output shapes here rather than during inference.
 
 		// Setup concrete output shapes only if all model output shapes are concretes, otherwise it will be set during Run()
-		for (NNECore::FTensorDesc SymbolicTensorDesc : OutputSymbolicTensors)
+		for (NNE::FTensorDesc SymbolicTensorDesc : OutputSymbolicTensors)
 		{
 			if (SymbolicTensorDesc.GetShape().IsConcrete())
 			{
-				NNECore::Internal::FTensor Tensor = NNECore::Internal::FTensor::MakeFromSymbolicDesc(SymbolicTensorDesc);
+				NNE::Internal::FTensor Tensor = NNE::Internal::FTensor::MakeFromSymbolicDesc(SymbolicTensorDesc);
 				OutputTensors.Emplace(Tensor);
 				OutputTensorShapes.Emplace(Tensor.GetShape());
 			}
@@ -229,7 +229,7 @@ namespace UE::NNERuntimeORT::Private
 		return 0;
 	}
 
-	int FModelInstanceORT::RunSync(TConstArrayView<NNECore::FTensorBindingGPU> InInputBindings, TConstArrayView<NNECore::FTensorBindingGPU> InOutputBindings)
+	int FModelInstanceORT::RunSync(TConstArrayView<NNE::FTensorBindingGPU> InInputBindings, TConstArrayView<NNE::FTensorBindingGPU> InOutputBindings)
 	{
 		DECLARE_SCOPE_CYCLE_COUNTER(TEXT("FModelORT_Run"), STAT_FModelORT_Run, STATGROUP_NNE);
 
@@ -328,7 +328,7 @@ namespace UE::NNERuntimeORT::Private
 
 #if PLATFORM_WINDOWS
 
-	TUniquePtr<UE::NNECore::IModelInstanceGPU> FModelORTDml::CreateModelInstance()
+	TUniquePtr<UE::NNE::IModelInstanceGPU> FModelORTDml::CreateModelInstance()
 	{
 		const FRuntimeConf InConf;
 		FModelInstanceORTDml* ModelInstance = new FModelInstanceORTDml(ORTEnvironment, InConf);
@@ -336,11 +336,11 @@ namespace UE::NNERuntimeORT::Private
 		if (!ModelInstance->Init(ModelData))
 		{
 			delete ModelInstance;
-			return TUniquePtr<UE::NNECore::IModelInstanceGPU>();
+			return TUniquePtr<UE::NNE::IModelInstanceGPU>();
 		}
 
-		UE::NNECore::IModelInstanceGPU* IModelInstance = static_cast<UE::NNECore::IModelInstanceGPU*>(ModelInstance);
-		return TUniquePtr<UE::NNECore::IModelInstanceGPU>(IModelInstance);
+		UE::NNE::IModelInstanceGPU* IModelInstance = static_cast<UE::NNE::IModelInstanceGPU*>(ModelInstance);
+		return TUniquePtr<UE::NNE::IModelInstanceGPU>(IModelInstance);
 	}
 
 	FModelORTDml::FModelORTDml(Ort::Env* InORTEnvironment, TConstArrayView<uint8> InModelData) :
@@ -350,7 +350,7 @@ namespace UE::NNERuntimeORT::Private
 
 	}
 
-	TUniquePtr<UE::NNECore::IModelInstanceGPU> FModelORTCuda::CreateModelInstance()
+	TUniquePtr<UE::NNE::IModelInstanceGPU> FModelORTCuda::CreateModelInstance()
 	{
 		const FRuntimeConf InConf;
 		FModelInstanceORTCuda* ModelInstance = new FModelInstanceORTCuda(ORTEnvironment, InConf);
@@ -358,11 +358,11 @@ namespace UE::NNERuntimeORT::Private
 		if (!ModelInstance->Init(ModelData))
 		{
 			delete ModelInstance;
-			return TUniquePtr<UE::NNECore::IModelInstanceGPU>();
+			return TUniquePtr<UE::NNE::IModelInstanceGPU>();
 		}
 
-		UE::NNECore::IModelInstanceGPU* IModelInstance = static_cast<UE::NNECore::IModelInstanceGPU*>(ModelInstance);
-		return TUniquePtr<UE::NNECore::IModelInstanceGPU>(IModelInstance);
+		UE::NNE::IModelInstanceGPU* IModelInstance = static_cast<UE::NNE::IModelInstanceGPU*>(ModelInstance);
+		return TUniquePtr<UE::NNE::IModelInstanceGPU>(IModelInstance);
 	}
 
 	FModelORTCuda::FModelORTCuda(Ort::Env* InORTEnvironment, TConstArrayView<uint8> InModelData) :
diff --git a/Engine/Plugins/Experimental/NNERuntimeORTGpu/Source/NNERuntimeORT/Private/NNERuntimeORTModel.h b/Engine/Plugins/Experimental/NNERuntimeORTGpu/Source/NNERuntimeORT/Private/NNERuntimeORTModel.h
--- a/Engine/Plugins/Experimental/NNERuntimeORTGpu/Source/NNERuntimeORT/Private/NNERuntimeORTModel.h
+++ b/Engine/Plugins/Experimental/NNERuntimeORTGpu/Source/NNERuntimeORT/Private/NNERuntimeORTModel.h
@@ -25,7 +25,7 @@ namespace UE::NNERuntimeORT::Private
 		EThreadPriority ThreadPriority = EThreadPriority::TPri_Normal;
 	};
 
-	class FModelInstanceORT : public NNECore::Internal::FModelInstanceBase<NNECore::IModelInstanceGPU>
+	class FModelInstanceORT : public NNE::Internal::FModelInstanceBase<NNE::IModelInstanceGPU>
 	{
 
 	public:
@@ -33,8 +33,8 @@ namespace UE::NNERuntimeORT::Private
 		FModelInstanceORT(Ort::Env* InORTEnvironment, const FRuntimeConf& InRuntimeConf);
 		virtual ~FModelInstanceORT() = default;
 
-		virtual int SetInputTensorShapes(TConstArrayView<NNECore::FTensorShape> InInputShapes) override;
-		virtual int RunSync(TConstArrayView<NNECore::FTensorBindingGPU> InInputBindings, TConstArrayView<NNECore::FTensorBindingGPU> InOutputBindings) override;
+		virtual int SetInputTensorShapes(TConstArrayView<NNE::FTensorShape> InInputShapes) override;
+		virtual int RunSync(TConstArrayView<NNE::FTensorBindingGPU> InInputBindings, TConstArrayView<NNE::FTensorBindingGPU> InOutputBindings) override;
 
 		bool Init(TConstArrayView<uint8> ModelData);
 		bool IsLoaded() const;
@@ -66,8 +66,8 @@ namespace UE::NNERuntimeORT::Private
 		TArray<char*> InputTensorNames;
 		TArray<char*> OutputTensorNames;
 
-		TArray<NNECore::Internal::FTensor> InputTensors;
-		TArray<NNECore::Internal::FTensor> OutputTensors;
+		TArray<NNE::Internal::FTensor> InputTensors;
+		TArray<NNE::Internal::FTensor> OutputTensors;
 	
 		virtual bool InitializedAndConfigureMembers();
 		bool ConfigureTensors(const bool InIsInput);
@@ -101,26 +101,26 @@ namespace UE::NNERuntimeORT::Private
 		virtual bool InitializedAndConfigureMembers() override;
 	};
 
-	class FModelORTDml : public NNECore::IModelGPU
+	class FModelORTDml : public NNE::IModelGPU
 	{
 	public:
 		FModelORTDml(Ort::Env* InORTEnvironment, TConstArrayView<uint8> ModelData);
 		virtual ~FModelORTDml() {};
 
-		virtual TUniquePtr<UE::NNECore::IModelInstanceGPU> CreateModelInstance() override;
+		virtual TUniquePtr<UE::NNE::IModelInstanceGPU> CreateModelInstance() override;
 
 	private:
 		Ort::Env* ORTEnvironment;
 		TArray<uint8> ModelData;
 	};
 
-	class FModelORTCuda : public NNECore::IModelGPU
+	class FModelORTCuda : public NNE::IModelGPU
 	{
 	public:
 		FModelORTCuda(Ort::Env* InORTEnvironment, TConstArrayView<uint8> ModelData);
 		virtual ~FModelORTCuda() {};
 
-		virtual TUniquePtr<UE::NNECore::IModelInstanceGPU> CreateModelInstance() override;
+		virtual TUniquePtr<UE::NNE::IModelInstanceGPU> CreateModelInstance() override;
 
 	private:
 		Ort::Env* ORTEnvironment;
diff --git a/Engine/Plugins/Experimental/NNERuntimeORTGpu/Source/NNERuntimeORT/Private/NNERuntimeORTModule.cpp b/Engine/Plugins/Experimental/NNERuntimeORTGpu/Source/NNERuntimeORT/Private/NNERuntimeORTModule.cpp
--- a/Engine/Plugins/Experimental/NNERuntimeORTGpu/Source/NNERuntimeORT/Private/NNERuntimeORTModule.cpp
+++ b/Engine/Plugins/Experimental/NNERuntimeORTGpu/Source/NNERuntimeORT/Private/NNERuntimeORTModule.cpp
@@ -68,7 +68,7 @@ void FNNERuntimeORTModule::StartupModule()
 
 		NNERuntimeORTCuda->Init(ENNERuntimeORTGpuProvider::Cuda);
 		NNERuntimeORTCuda->AddToRoot();
-		UE::NNECore::RegisterRuntime(RuntimeCudaInterface);
+		UE::NNE::RegisterRuntime(RuntimeCudaInterface);
 	}
 	
 	// NNE runtime ORT Dml startup
@@ -79,7 +79,7 @@ void FNNERuntimeORTModule::StartupModule()
 
 		NNERuntimeORTDml->Init(ENNERuntimeORTGpuProvider::Dml);
 		NNERuntimeORTDml->AddToRoot();
-		UE::NNECore::RegisterRuntime(RuntimeDmlInterface);
+		UE::NNE::RegisterRuntime(RuntimeDmlInterface);
 	}
 #endif
 }
@@ -92,7 +92,7 @@ void FNNERuntimeORTModule::ShutdownModule()
 	{
 		TWeakInterfacePtr<INNERuntime> RuntimeDmlInterface(NNERuntimeORTDml.Get());
 
-		UE::NNECore::UnregisterRuntime(RuntimeDmlInterface);
+		UE::NNE::UnregisterRuntime(RuntimeDmlInterface);
 		NNERuntimeORTDml->RemoveFromRoot();
 		NNERuntimeORTDml = TWeakObjectPtr<UNNERuntimeORTGpuImpl>(nullptr);
 	}
@@ -102,7 +102,7 @@ void FNNERuntimeORTModule::ShutdownModule()
 	{
 		TWeakInterfacePtr<INNERuntime> RuntimeCudaInterface(NNERuntimeORTCuda.Get());
 
-		UE::NNECore::UnregisterRuntime(RuntimeCudaInterface);
+		UE::NNE::UnregisterRuntime(RuntimeCudaInterface);
 		NNERuntimeORTCuda->RemoveFromRoot();
 		NNERuntimeORTCuda = TWeakObjectPtr<UNNERuntimeORTGpuImpl>(nullptr);
 	}
diff --git a/Engine/Plugins/Experimental/NNERuntimeORTGpu/Source/NNERuntimeORT/Private/NNERuntimeORTUtils.h b/Engine/Plugins/Experimental/NNERuntimeORTGpu/Source/NNERuntimeORT/Private/NNERuntimeORTUtils.h
--- a/Engine/Plugins/Experimental/NNERuntimeORTGpu/Source/NNERuntimeORT/Private/NNERuntimeORTUtils.h
+++ b/Engine/Plugins/Experimental/NNERuntimeORTGpu/Source/NNERuntimeORT/Private/NNERuntimeORTUtils.h
@@ -118,8 +118,8 @@ namespace UE::NNERuntimeORT::Private
 	}
 
 	inline void BindTensorsToORT(
-		TConstArrayView<NNECore::FTensorBindingGPU> InBindingTensors,
-		TConstArrayView<NNECore::Internal::FTensor> InTensors,
+		TConstArrayView<NNE::FTensorBindingGPU> InBindingTensors,
+		TConstArrayView<NNE::Internal::FTensor> InTensors,
 		TConstArrayView<ONNXTensorElementDataType> InTensorsORTType,
 		const Ort::MemoryInfo* InAllocatorInfo,
 		TArray<Ort::Value>& OutOrtTensors
@@ -136,8 +136,8 @@ namespace UE::NNERuntimeORT::Private
 
 		for (uint32 Index = 0; Index < NumBinding; ++Index)
 		{
-			const NNECore::FTensorBindingGPU& Binding = InBindingTensors[Index];
-			const NNECore::Internal::FTensor& Tensor = InTensors[Index];
+			const NNE::FTensorBindingGPU& Binding = InBindingTensors[Index];
+			const NNE::Internal::FTensor& Tensor = InTensors[Index];
 			const ONNXTensorElementDataType CurrentORTType = InTensorsORTType[Index];
 
 			TUniquePtr<int64_t[]> SizesInt64t;
@@ -164,9 +164,9 @@ namespace UE::NNERuntimeORT::Private
 
 	inline void CopyFromORTToBindings(
 		TConstArrayView<Ort::Value> InOrtTensors,
-		TConstArrayView<NNECore::FTensorBindingGPU> InBindingTensors,
-		TConstArrayView<NNECore::FTensorDesc> InTensorDescs,
-		TArray<NNECore::Internal::FTensor>& OutTensors
+		TConstArrayView<NNE::FTensorBindingGPU> InBindingTensors,
+		TConstArrayView<NNE::FTensorDesc> InTensorDescs,
+		TArray<NNE::Internal::FTensor>& OutTensors
 	)
 	{
 		const uint32 NumBinding = InOrtTensors.Num();
@@ -180,8 +180,8 @@ namespace UE::NNERuntimeORT::Private
 
 		for (uint32 Index = 0; Index < NumBinding; ++Index)
 		{
-			const NNECore::FTensorBindingGPU& Binding = InBindingTensors[Index];
-			const NNECore::FTensorDesc& TensorDesc = InTensorDescs[Index];
+			const NNE::FTensorBindingGPU& Binding = InBindingTensors[Index];
+			const NNE::FTensorDesc& TensorDesc = InTensorDescs[Index];
 			const Ort::Value& OrtTensor = InOrtTensors[Index];
 			const std::vector<int64_t>& OrtShape = OrtTensor.GetTensorTypeAndShapeInfo().GetShape();
 
@@ -192,8 +192,8 @@ namespace UE::NNERuntimeORT::Private
 				ShapeData.Add(OrtShape[DimIndex]);
 			}
 
-			NNECore::FTensorShape Shape = NNECore::FTensorShape::Make(ShapeData);
-			NNECore::Internal::FTensor Tensor = NNECore::Internal::FTensor::Make(TensorDesc.GetName(), Shape, TensorDesc.GetDataType());
+			NNE::FTensorShape Shape = NNE::FTensorShape::Make(ShapeData);
+			NNE::Internal::FTensor Tensor = NNE::Internal::FTensor::Make(TensorDesc.GetName(), Shape, TensorDesc.GetDataType());
 			OutTensors.Emplace(Tensor);
 
 			const void* OrtTensorData = OrtTensor.GetTensorData<void>();
diff --git a/Engine/Plugins/Experimental/NNERuntimeRDG/Source/NNEHlslShaders/Internal/NNEHlslShadersElementWiseBinaryCS.h b/Engine/Plugins/Experimental/NNERuntimeRDG/Source/NNEHlslShaders/Internal/NNEHlslShadersElementWiseBinaryCS.h
--- a/Engine/Plugins/Experimental/NNERuntimeRDG/Source/NNEHlslShaders/Internal/NNEHlslShadersElementWiseBinaryCS.h
+++ b/Engine/Plugins/Experimental/NNERuntimeRDG/Source/NNEHlslShaders/Internal/NNEHlslShadersElementWiseBinaryCS.h
@@ -22,7 +22,7 @@ namespace UE::NNEHlslShaders::Internal
 		DECLARE_GLOBAL_SHADER(TElementWiseBinaryCS);
 		SHADER_USE_PARAMETER_STRUCT(TElementWiseBinaryCS, FHlslShaderBase)
 
-		class FOperatorType : SHADER_PERMUTATION_ENUM_CLASS("OP_TYPENAME", NNECore::Internal::EElementWiseBinaryOperatorType);
+		class FOperatorType : SHADER_PERMUTATION_ENUM_CLASS("OP_TYPENAME", NNE::Internal::EElementWiseBinaryOperatorType);
 		class FBinaryNumDimensions : SHADER_PERMUTATION_RANGE_INT("NUM_DIMENSIONS", 1, FElementWiseBinaryConstants::MAX_NUM_DIMENSIONS);
 		using FPermutationDomain = TShaderPermutationDomain<FOperatorType, FBinaryNumDimensions>;
 
@@ -41,6 +41,6 @@ namespace UE::NNEHlslShaders::Internal
 
 	private:
 
-		static const FString GetOpFunc(NNECore::Internal::EElementWiseBinaryOperatorType OpType);
+		static const FString GetOpFunc(NNE::Internal::EElementWiseBinaryOperatorType OpType);
 	};
 } // UE::NNEHlslShaders::Internal
diff --git a/Engine/Plugins/Experimental/NNERuntimeRDG/Source/NNEHlslShaders/Internal/NNEHlslShadersElementWiseUnaryCS.h b/Engine/Plugins/Experimental/NNERuntimeRDG/Source/NNEHlslShaders/Internal/NNEHlslShadersElementWiseUnaryCS.h
--- a/Engine/Plugins/Experimental/NNERuntimeRDG/Source/NNEHlslShaders/Internal/NNEHlslShadersElementWiseUnaryCS.h
+++ b/Engine/Plugins/Experimental/NNERuntimeRDG/Source/NNEHlslShaders/Internal/NNEHlslShadersElementWiseUnaryCS.h
@@ -21,7 +21,7 @@ namespace UE::NNEHlslShaders::Internal
 		DECLARE_GLOBAL_SHADER(TElementWiseUnaryCS);
 		SHADER_USE_PARAMETER_STRUCT(TElementWiseUnaryCS, FHlslShaderBase)
 
-		class FOperatorType : SHADER_PERMUTATION_ENUM_CLASS("OP_TYPENAME", NNECore::Internal::EElementWiseUnaryOperatorType);
+		class FOperatorType : SHADER_PERMUTATION_ENUM_CLASS("OP_TYPENAME", NNE::Internal::EElementWiseUnaryOperatorType);
 		using FPermutationDomain = TShaderPermutationDomain<FOperatorType>;
 
 	public:
@@ -40,6 +40,6 @@ namespace UE::NNEHlslShaders::Internal
 
 	private:
 
-		static const FString GetOpFunc(NNECore::Internal::EElementWiseUnaryOperatorType OpType);
+		static const FString GetOpFunc(NNE::Internal::EElementWiseUnaryOperatorType OpType);
 	};
 } // UE::NNEHlslShaders::Internal
diff --git a/Engine/Plugins/Experimental/NNERuntimeRDG/Source/NNEHlslShaders/Internal/NNEHlslShadersElementWiseVariadicCS.h b/Engine/Plugins/Experimental/NNERuntimeRDG/Source/NNEHlslShaders/Internal/NNEHlslShadersElementWiseVariadicCS.h
--- a/Engine/Plugins/Experimental/NNERuntimeRDG/Source/NNEHlslShaders/Internal/NNEHlslShadersElementWiseVariadicCS.h
+++ b/Engine/Plugins/Experimental/NNERuntimeRDG/Source/NNEHlslShaders/Internal/NNEHlslShadersElementWiseVariadicCS.h
@@ -23,7 +23,7 @@ namespace UE::NNEHlslShaders::Internal
 		DECLARE_GLOBAL_SHADER(TElementWiseVariadicCS);
 		SHADER_USE_PARAMETER_STRUCT(TElementWiseVariadicCS, FHlslShaderBase)
 
-		class FOperatorType : SHADER_PERMUTATION_ENUM_CLASS("OP_TYPENAME", NNECore::Internal::EElementWiseVariadicOperatorType);
+		class FOperatorType : SHADER_PERMUTATION_ENUM_CLASS("OP_TYPENAME", NNE::Internal::EElementWiseVariadicOperatorType);
 		class FApplyScale : SHADER_PERMUTATION_BOOL("APPLYSCALE");
 		class FOutputAsInput : SHADER_PERMUTATION_BOOL("OUTPUTASINPUT");
 		class FNumInput : SHADER_PERMUTATION_RANGE_INT("NUMINPUT", 1, FElementWiseVariadicConstants::MAX_NUM_INPUT);
@@ -49,6 +49,6 @@ namespace UE::NNEHlslShaders::Internal
 
 	private:
 
-		static const FString GetOpFunc(NNECore::Internal::EElementWiseVariadicOperatorType OpType);
+		static const FString GetOpFunc(NNE::Internal::EElementWiseVariadicOperatorType OpType);
 	};
 } // UE::NNEHlslShaders::Internal
\ No newline at end of file
diff --git a/Engine/Plugins/Experimental/NNERuntimeRDG/Source/NNEHlslShaders/Internal/NNEHlslShadersGatherCS.h b/Engine/Plugins/Experimental/NNERuntimeRDG/Source/NNEHlslShaders/Internal/NNEHlslShadersGatherCS.h
--- a/Engine/Plugins/Experimental/NNERuntimeRDG/Source/NNEHlslShaders/Internal/NNEHlslShadersGatherCS.h
+++ b/Engine/Plugins/Experimental/NNERuntimeRDG/Source/NNEHlslShaders/Internal/NNEHlslShadersGatherCS.h
@@ -6,7 +6,7 @@
 #include "RenderGraphUtils.h"
 #include "ShaderParameterUtils.h"
 
-namespace UE::NNECore::Internal { class FTensor; }
+namespace UE::NNE::Internal { class FTensor; }
 
 namespace UE::NNEHlslShaders::Internal
 {
@@ -43,7 +43,7 @@ namespace UE::NNEHlslShaders::Internal
 
 		static void ModifyCompilationEnvironment(const FGlobalShaderPermutationParameters& InParameters, FShaderCompilerEnvironment& OutEnvironment);
 
-		static void FillInParameters(int32 Axis, const NNECore::Internal::FTensor& Data, const NNECore::Internal::FTensor& Indices, FParameters& Parameters);
+		static void FillInParameters(int32 Axis, const NNE::Internal::FTensor& Data, const NNE::Internal::FTensor& Indices, FParameters& Parameters);
 
 		static FIntVector GetGroupCount(const FParameters& Parameters);
 	};
diff --git a/Engine/Plugins/Experimental/NNERuntimeRDG/Source/NNEHlslShaders/Internal/NNEHlslShadersGemmCS.h b/Engine/Plugins/Experimental/NNERuntimeRDG/Source/NNEHlslShaders/Internal/NNEHlslShadersGemmCS.h
--- a/Engine/Plugins/Experimental/NNERuntimeRDG/Source/NNEHlslShaders/Internal/NNEHlslShadersGemmCS.h
+++ b/Engine/Plugins/Experimental/NNERuntimeRDG/Source/NNEHlslShaders/Internal/NNEHlslShadersGemmCS.h
@@ -6,7 +6,7 @@
 #include "RenderGraphUtils.h"
 #include "ShaderParameterUtils.h"
 
-namespace UE::NNECore::Internal { class FTensor; }
+namespace UE::NNE::Internal { class FTensor; }
 
 namespace UE::NNEHlslShaders::Internal
 {
@@ -80,9 +80,9 @@ namespace UE::NNEHlslShaders::Internal
 
 		static void ModifyCompilationEnvironment(const FGlobalShaderPermutationParameters& InParameters, FShaderCompilerEnvironment& OutEnvironment);
 
-		static void FillInParameters(float Alpha, float Beta, int32 TransA, int32 TransB, const NNECore::Internal::FTensor& InputA, const NNECore::Internal::FTensor& InputB, const NNECore::Internal::FTensor* InputC, float CScalar, FParameters& Parameters);
+		static void FillInParameters(float Alpha, float Beta, int32 TransA, int32 TransB, const NNE::Internal::FTensor& InputA, const NNE::Internal::FTensor& InputB, const NNE::Internal::FTensor* InputC, float CScalar, FParameters& Parameters);
 
-		static void FillInParametersMatMul(const NNECore::Internal::FTensor& InputA, const NNECore::Internal::FTensor& InputB, FParameters& Parameters);
+		static void FillInParametersMatMul(const NNE::Internal::FTensor& InputA, const NNE::Internal::FTensor& InputB, FParameters& Parameters);
 
 		static FIntVector GetGroupCount(const FParameters& Parameters, EGemmAlgorithm Algorithm, int32 NumStackDimensions);
 		static EGemmAlgorithm GetAlgorithm(const FParameters& Parameters);
diff --git a/Engine/Plugins/Experimental/NNERuntimeRDG/Source/NNEHlslShaders/Internal/NNEHlslShadersInstanceNormalizationCS.h b/Engine/Plugins/Experimental/NNERuntimeRDG/Source/NNEHlslShaders/Internal/NNEHlslShadersInstanceNormalizationCS.h
--- a/Engine/Plugins/Experimental/NNERuntimeRDG/Source/NNEHlslShaders/Internal/NNEHlslShadersInstanceNormalizationCS.h
+++ b/Engine/Plugins/Experimental/NNERuntimeRDG/Source/NNEHlslShaders/Internal/NNEHlslShadersInstanceNormalizationCS.h
@@ -6,7 +6,7 @@
 #include "RenderGraphUtils.h"
 #include "ShaderParameterUtils.h"
 
-namespace UE::NNECore::Internal { class FTensor; }
+namespace UE::NNE::Internal { class FTensor; }
 
 namespace UE::NNEHlslShaders::Internal
 {
@@ -41,7 +41,7 @@ namespace UE::NNEHlslShaders::Internal
 			SHADER_PARAMETER_RDG_BUFFER_UAV(RWBuffer<float>, Output)
 		END_SHADER_PARAMETER_STRUCT()
 
-		static void FillInParameters(float Epsilon, const NNECore::Internal::FTensor& Input, FParameters& Parameters);
+		static void FillInParameters(float Epsilon, const NNE::Internal::FTensor& Input, FParameters& Parameters);
 
 		static FIntVector GetGroupCount(const FParameters& Parameters, EInstanceNormalizationAlgorithm Algorithm);
 		static EInstanceNormalizationAlgorithm GetAlgorithm(const FParameters& Parameters);
diff --git a/Engine/Plugins/Experimental/NNERuntimeRDG/Source/NNEHlslShaders/Private/NNEHlslShadersElementWiseBinaryCS.cpp b/Engine/Plugins/Experimental/NNERuntimeRDG/Source/NNEHlslShaders/Private/NNEHlslShadersElementWiseBinaryCS.cpp
--- a/Engine/Plugins/Experimental/NNERuntimeRDG/Source/NNEHlslShaders/Private/NNEHlslShadersElementWiseBinaryCS.cpp
+++ b/Engine/Plugins/Experimental/NNERuntimeRDG/Source/NNEHlslShaders/Private/NNEHlslShadersElementWiseBinaryCS.cpp
@@ -18,16 +18,16 @@ namespace UE::NNEHlslShaders::Internal
 		OutEnvironment.SetDefine(TEXT("ELEMENTWISE_OP(X,Y)"), *OpFunc);
 	}
 
-	const FString TElementWiseBinaryCS::GetOpFunc(NNECore::Internal::EElementWiseBinaryOperatorType OpType)
+	const FString TElementWiseBinaryCS::GetOpFunc(NNE::Internal::EElementWiseBinaryOperatorType OpType)
 	{
-		FString OpTable[(int32) NNECore::Internal::EElementWiseBinaryOperatorType::MAX];
+		FString OpTable[(int32) NNE::Internal::EElementWiseBinaryOperatorType::MAX];
 
-		for (int32 Idx = 0; Idx < (int32) NNECore::Internal::EElementWiseBinaryOperatorType::MAX; ++Idx)
+		for (int32 Idx = 0; Idx < (int32) NNE::Internal::EElementWiseBinaryOperatorType::MAX; ++Idx)
 		{
 			OpTable[Idx] = FString("");
 		}
 
-#define OP(OpName, OpFunc) OpTable[(int32) NNECore::Internal::EElementWiseBinaryOperatorType::OpName] = OpFunc
+#define OP(OpName, OpFunc) OpTable[(int32) NNE::Internal::EElementWiseBinaryOperatorType::OpName] = OpFunc
 		OP(Add,              TEXT("((X)+(Y))"));
 		//OP(And,            TEXT("((X)&&(Y))"));
 		OP(Div,              TEXT("((X)/(Y))"));
diff --git a/Engine/Plugins/Experimental/NNERuntimeRDG/Source/NNEHlslShaders/Private/NNEHlslShadersElementWiseUnaryCS.cpp b/Engine/Plugins/Experimental/NNERuntimeRDG/Source/NNEHlslShaders/Private/NNEHlslShadersElementWiseUnaryCS.cpp
--- a/Engine/Plugins/Experimental/NNERuntimeRDG/Source/NNEHlslShaders/Private/NNEHlslShadersElementWiseUnaryCS.cpp
+++ b/Engine/Plugins/Experimental/NNERuntimeRDG/Source/NNEHlslShaders/Private/NNEHlslShadersElementWiseUnaryCS.cpp
@@ -18,16 +18,16 @@ namespace UE::NNEHlslShaders::Internal
 		OutEnvironment.SetDefine(TEXT("ELEMENTWISE_OP(X)"), *OpFunc);
 	}
 
-	const FString TElementWiseUnaryCS::GetOpFunc(NNECore::Internal::EElementWiseUnaryOperatorType OpType)
+	const FString TElementWiseUnaryCS::GetOpFunc(NNE::Internal::EElementWiseUnaryOperatorType OpType)
 	{
-		FString OpTable[(int32) NNECore::Internal::EElementWiseUnaryOperatorType::MAX];
+		FString OpTable[(int32) NNE::Internal::EElementWiseUnaryOperatorType::MAX];
 
-		for (int32 Idx = 0; Idx < (int32) NNECore::Internal::EElementWiseUnaryOperatorType::MAX; ++Idx)
+		for (int32 Idx = 0; Idx < (int32) NNE::Internal::EElementWiseUnaryOperatorType::MAX; ++Idx)
 		{
 			OpTable[Idx] = FString("");
 		}
 
-#define OP(OpName, OpFunc) OpTable[(int32) NNECore::Internal::EElementWiseUnaryOperatorType::OpName] = OpFunc
+#define OP(OpName, OpFunc) OpTable[(int32) NNE::Internal::EElementWiseUnaryOperatorType::OpName] = OpFunc
 		OP(Abs,         TEXT("abs(X)"));
 		OP(Acos,        TEXT("acos(X)"));
 		OP(Acosh,       TEXT("acosh(X)"));
diff --git a/Engine/Plugins/Experimental/NNERuntimeRDG/Source/NNEHlslShaders/Private/NNEHlslShadersElementWiseVariadicCS.cpp b/Engine/Plugins/Experimental/NNERuntimeRDG/Source/NNEHlslShaders/Private/NNEHlslShadersElementWiseVariadicCS.cpp
--- a/Engine/Plugins/Experimental/NNERuntimeRDG/Source/NNEHlslShaders/Private/NNEHlslShadersElementWiseVariadicCS.cpp
+++ b/Engine/Plugins/Experimental/NNERuntimeRDG/Source/NNEHlslShaders/Private/NNEHlslShadersElementWiseVariadicCS.cpp
@@ -18,16 +18,16 @@ namespace UE::NNEHlslShaders::Internal
 		OutEnvironment.SetDefine(TEXT("ELEMENTWISE_OP(X,Y)"), *OpFunc);
 	}
 
-	const FString TElementWiseVariadicCS::GetOpFunc(NNECore::Internal::EElementWiseVariadicOperatorType OpType)
+	const FString TElementWiseVariadicCS::GetOpFunc(NNE::Internal::EElementWiseVariadicOperatorType OpType)
 	{
-		FString OpTable[(int32) NNECore::Internal::EElementWiseVariadicOperatorType::MAX];
+		FString OpTable[(int32) NNE::Internal::EElementWiseVariadicOperatorType::MAX];
 
-		for (int32 Idx = 0; Idx < (int32) NNECore::Internal::EElementWiseVariadicOperatorType::MAX; ++Idx)
+		for (int32 Idx = 0; Idx < (int32) NNE::Internal::EElementWiseVariadicOperatorType::MAX; ++Idx)
 		{
 			OpTable[Idx] = FString("");
 		}
 
-#define OP(OpName, OpFunc) OpTable[(int32) NNECore::Internal::EElementWiseVariadicOperatorType::OpName] = OpFunc
+#define OP(OpName, OpFunc) OpTable[(int32) NNE::Internal::EElementWiseVariadicOperatorType::OpName] = OpFunc
 		OP(Max, TEXT("max(X,Y)"));
 		OP(Min, TEXT("min(X,Y)"));
 		OP(Mean, TEXT("((X)+(Y))"));
diff --git a/Engine/Plugins/Experimental/NNERuntimeRDG/Source/NNEHlslShaders/Private/NNEHlslShadersGatherCS.cpp b/Engine/Plugins/Experimental/NNERuntimeRDG/Source/NNEHlslShaders/Private/NNEHlslShadersGatherCS.cpp
--- a/Engine/Plugins/Experimental/NNERuntimeRDG/Source/NNEHlslShaders/Private/NNEHlslShadersGatherCS.cpp
+++ b/Engine/Plugins/Experimental/NNERuntimeRDG/Source/NNEHlslShaders/Private/NNEHlslShadersGatherCS.cpp
@@ -16,7 +16,7 @@ namespace UE::NNEHlslShaders::Internal
 	}
 
 	//template <typename DataElementType, typename IndicesElementType>
-	void TGatherCS::FillInParameters(int32 Axis, const NNECore::Internal::FTensor& Data, const NNECore::Internal::FTensor& Indices, FParameters& Parameters)
+	void TGatherCS::FillInParameters(int32 Axis, const NNE::Internal::FTensor& Data, const NNE::Internal::FTensor& Indices, FParameters& Parameters)
 	{
 		Parameters.Axis = Axis;
 
diff --git a/Engine/Plugins/Experimental/NNERuntimeRDG/Source/NNEHlslShaders/Private/NNEHlslShadersGemmCS.cpp b/Engine/Plugins/Experimental/NNERuntimeRDG/Source/NNEHlslShaders/Private/NNEHlslShadersGemmCS.cpp
--- a/Engine/Plugins/Experimental/NNERuntimeRDG/Source/NNEHlslShaders/Private/NNEHlslShadersGemmCS.cpp
+++ b/Engine/Plugins/Experimental/NNERuntimeRDG/Source/NNEHlslShaders/Private/NNEHlslShadersGemmCS.cpp
@@ -18,7 +18,7 @@ namespace UE::NNEHlslShaders::Internal
 		TArray<uint32> StackStrideB;
 	};
 
-	static GemmMatrixParameters GetMatrixParametersMatMul(const NNECore::Internal::FTensor& InputA, const NNECore::Internal::FTensor& InputB)
+	static GemmMatrixParameters GetMatrixParametersMatMul(const NNE::Internal::FTensor& InputA, const NNE::Internal::FTensor& InputB)
 	{
 		check(InputA.GetShape().Rank() != 0);
 		check(InputB.GetShape().Rank() != 0);
@@ -81,7 +81,7 @@ namespace UE::NNEHlslShaders::Internal
 		OutEnvironment.SetDefine(TEXT("MAX_NUM_STACK_DIMENSIONS"), FGemmConstants::MAX_NUM_STACK_DIMENSIONS);
 	}
 
-	void TGemmCS::FillInParameters(float Alpha, float Beta, int32 TransA, int32 TransB, const NNECore::Internal::FTensor& InputA, const NNECore::Internal::FTensor& InputB, const NNECore::Internal::FTensor* InputC, float CScalar, TGemmCS::FParameters& Parameters)
+	void TGemmCS::FillInParameters(float Alpha, float Beta, int32 TransA, int32 TransB, const NNE::Internal::FTensor& InputA, const NNE::Internal::FTensor& InputB, const NNE::Internal::FTensor* InputC, float CScalar, TGemmCS::FParameters& Parameters)
 	{
 		check(InputA.GetShape().Rank() == 2);
 		check(InputB.GetShape().Rank() == 2);
@@ -112,7 +112,7 @@ namespace UE::NNEHlslShaders::Internal
 		Parameters.CScalar = CScalar;
 	}
 
-	void TGemmCS::FillInParametersMatMul(const NNECore::Internal::FTensor& InputA, const NNECore::Internal::FTensor& InputB, TGemmCS::FParameters& Parameters)
+	void TGemmCS::FillInParametersMatMul(const NNE::Internal::FTensor& InputA, const NNE::Internal::FTensor& InputB, TGemmCS::FParameters& Parameters)
 	{
 		GemmMatrixParameters Params = GetMatrixParametersMatMul(InputA, InputB);
 
diff --git a/Engine/Plugins/Experimental/NNERuntimeRDG/Source/NNEHlslShaders/Private/NNEHlslShadersInstanceNormalizationCS.cpp b/Engine/Plugins/Experimental/NNERuntimeRDG/Source/NNEHlslShaders/Private/NNEHlslShadersInstanceNormalizationCS.cpp
--- a/Engine/Plugins/Experimental/NNERuntimeRDG/Source/NNEHlslShaders/Private/NNEHlslShadersInstanceNormalizationCS.cpp
+++ b/Engine/Plugins/Experimental/NNERuntimeRDG/Source/NNEHlslShaders/Private/NNEHlslShadersInstanceNormalizationCS.cpp
@@ -26,7 +26,7 @@ namespace UE::NNEHlslShaders::Internal
 		}
 	}
 
-	void TInstanceNormalizationCS::FillInParameters(float Epsilon, const  UE::NNECore::Internal::FTensor& Input, TInstanceNormalizationCS::FParameters& Parameters)
+	void TInstanceNormalizationCS::FillInParameters(float Epsilon, const  UE::NNE::Internal::FTensor& Input, TInstanceNormalizationCS::FParameters& Parameters)
 	{
 		check(Input.GetShape().Rank() >= 3);
 		
diff --git a/Engine/Plugins/Experimental/NNERuntimeRDG/Source/NNERuntimeRDG/Internal/NNERuntimeRDGHelperCast.h b/Engine/Plugins/Experimental/NNERuntimeRDG/Source/NNERuntimeRDG/Internal/NNERuntimeRDGHelperCast.h
--- a/Engine/Plugins/Experimental/NNERuntimeRDG/Source/NNERuntimeRDG/Internal/NNERuntimeRDGHelperCast.h
+++ b/Engine/Plugins/Experimental/NNERuntimeRDG/Source/NNERuntimeRDG/Internal/NNERuntimeRDGHelperCast.h
@@ -4,10 +4,10 @@
 
 #include "CoreMinimal.h"
 
-namespace UE::NNECore::Internal { class FTensor; }
+namespace UE::NNE::Internal { class FTensor; }
 
 namespace UE::NNERuntimeRDG::Internal::CPUHelper::Cast
 {
-	void NNERUNTIMERDG_API Apply(const NNECore::Internal::FTensor& Tensor, NNECore::Internal::FTensor& OutputTensor);
+	void NNERUNTIMERDG_API Apply(const NNE::Internal::FTensor& Tensor, NNE::Internal::FTensor& OutputTensor);
 
 } // UE::NNERuntimeRDG::Internal::CPUHelper::Cast
\ No newline at end of file
diff --git a/Engine/Plugins/Experimental/NNERuntimeRDG/Source/NNERuntimeRDG/Internal/NNERuntimeRDGHelperConcat.h b/Engine/Plugins/Experimental/NNERuntimeRDG/Source/NNERuntimeRDG/Internal/NNERuntimeRDGHelperConcat.h
--- a/Engine/Plugins/Experimental/NNERuntimeRDG/Source/NNERuntimeRDG/Internal/NNERuntimeRDGHelperConcat.h
+++ b/Engine/Plugins/Experimental/NNERuntimeRDG/Source/NNERuntimeRDG/Internal/NNERuntimeRDGHelperConcat.h
@@ -7,5 +7,5 @@
 
 namespace UE::NNERuntimeRDG::Internal::CPUHelper::Concat
 {
-	void NNERUNTIMERDG_API Apply(TConstArrayView<NNECore::Internal::FTensorRef> InputTensors, NNECore::Internal::FTensor& OutputTensor, int32 Axis);
+	void NNERUNTIMERDG_API Apply(TConstArrayView<NNE::Internal::FTensorRef> InputTensors, NNE::Internal::FTensor& OutputTensor, int32 Axis);
 } // UE::NNERuntimeRDG::Internal::CPUHelper::Concat
\ No newline at end of file
diff --git a/Engine/Plugins/Experimental/NNERuntimeRDG/Source/NNERuntimeRDG/Internal/NNERuntimeRDGHelperElementWiseBinary.h b/Engine/Plugins/Experimental/NNERuntimeRDG/Source/NNERuntimeRDG/Internal/NNERuntimeRDGHelperElementWiseBinary.h
--- a/Engine/Plugins/Experimental/NNERuntimeRDG/Source/NNERuntimeRDG/Internal/NNERuntimeRDGHelperElementWiseBinary.h
+++ b/Engine/Plugins/Experimental/NNERuntimeRDG/Source/NNERuntimeRDG/Internal/NNERuntimeRDGHelperElementWiseBinary.h
@@ -5,10 +5,10 @@
 #include "CoreMinimal.h"
 #include "NNEOperator.h"
 
-namespace UE::NNECore::Internal { class FTensor; }
+namespace UE::NNE::Internal { class FTensor; }
 
 namespace UE::NNERuntimeRDG::Internal::CPUHelper::ElementWiseBinary
 {
-	void NNERUNTIMERDG_API Apply(NNECore::Internal::EElementWiseBinaryOperatorType OpType, const NNECore::Internal::FTensor& LHSTensor, const NNECore::Internal::FTensor& RHSTensor, NNECore::Internal::FTensor& OutputTensor);
+	void NNERUNTIMERDG_API Apply(NNE::Internal::EElementWiseBinaryOperatorType OpType, const NNE::Internal::FTensor& LHSTensor, const NNE::Internal::FTensor& RHSTensor, NNE::Internal::FTensor& OutputTensor);
 	
 } // UE::NNERuntimeRDG::Internal::CPUHelper::ElementWiseBinary
\ No newline at end of file
diff --git a/Engine/Plugins/Experimental/NNERuntimeRDG/Source/NNERuntimeRDG/Internal/NNERuntimeRDGHelperElementWiseUnary.h b/Engine/Plugins/Experimental/NNERuntimeRDG/Source/NNERuntimeRDG/Internal/NNERuntimeRDGHelperElementWiseUnary.h
--- a/Engine/Plugins/Experimental/NNERuntimeRDG/Source/NNERuntimeRDG/Internal/NNERuntimeRDGHelperElementWiseUnary.h
+++ b/Engine/Plugins/Experimental/NNERuntimeRDG/Source/NNERuntimeRDG/Internal/NNERuntimeRDGHelperElementWiseUnary.h
@@ -5,9 +5,9 @@
 #include "CoreMinimal.h"
 #include "NNEOperator.h"
 
-namespace UE::NNECore::Internal { class FTensor; }
+namespace UE::NNE::Internal { class FTensor; }
 
 namespace UE::NNERuntimeRDG::Internal::CPUHelper::ElementWiseUnary
 {
-	void NNERUNTIMERDG_API Apply(NNECore::Internal::EElementWiseUnaryOperatorType OpType, const NNECore::Internal::FTensor& Tensor, float Alpha, float Beta, float Gamma, NNECore::Internal::FTensor& OutputTensor);
+	void NNERUNTIMERDG_API Apply(NNE::Internal::EElementWiseUnaryOperatorType OpType, const NNE::Internal::FTensor& Tensor, float Alpha, float Beta, float Gamma, NNE::Internal::FTensor& OutputTensor);
 } // UE::NNERuntimeRDG::Internal::CPUHelper::ElementWiseUnary
\ No newline at end of file
diff --git a/Engine/Plugins/Experimental/NNERuntimeRDG/Source/NNERuntimeRDG/Internal/NNERuntimeRDGHelperGather.h b/Engine/Plugins/Experimental/NNERuntimeRDG/Source/NNERuntimeRDG/Internal/NNERuntimeRDGHelperGather.h
--- a/Engine/Plugins/Experimental/NNERuntimeRDG/Source/NNERuntimeRDG/Internal/NNERuntimeRDGHelperGather.h
+++ b/Engine/Plugins/Experimental/NNERuntimeRDG/Source/NNERuntimeRDG/Internal/NNERuntimeRDGHelperGather.h
@@ -4,9 +4,9 @@
 
 #include "CoreMinimal.h"
 
-namespace UE::NNECore::Internal { class FTensor; }
+namespace UE::NNE::Internal { class FTensor; }
 
 namespace UE::NNERuntimeRDG::Internal::CPUHelper::Gather
 {
-	void NNERUNTIMERDG_API Apply(const NNECore::Internal::FTensor& DataTensor, const NNECore::Internal::FTensor& IndicesTensor, int32 Axis, NNECore::Internal::FTensor& OutputTensor);
+	void NNERUNTIMERDG_API Apply(const NNE::Internal::FTensor& DataTensor, const NNE::Internal::FTensor& IndicesTensor, int32 Axis, NNE::Internal::FTensor& OutputTensor);
 } // UE::NNERuntimeRDG::Internal::CPUHelper::Gather
\ No newline at end of file
diff --git a/Engine/Plugins/Experimental/NNERuntimeRDG/Source/NNERuntimeRDG/Internal/NNERuntimeRDGHelperSlice.h b/Engine/Plugins/Experimental/NNERuntimeRDG/Source/NNERuntimeRDG/Internal/NNERuntimeRDGHelperSlice.h
--- a/Engine/Plugins/Experimental/NNERuntimeRDG/Source/NNERuntimeRDG/Internal/NNERuntimeRDGHelperSlice.h
+++ b/Engine/Plugins/Experimental/NNERuntimeRDG/Source/NNERuntimeRDG/Internal/NNERuntimeRDGHelperSlice.h
@@ -4,9 +4,9 @@
 
 #include "CoreMinimal.h"
 
-namespace UE::NNECore::Internal { class FTensor; }
+namespace UE::NNE::Internal { class FTensor; }
 
 namespace UE::NNERuntimeRDG::Internal::CPUHelper::Slice
 {
-	void NNERUNTIMERDG_API Apply(const NNECore::Internal::FTensor& InputTensor, NNECore::Internal::FTensor& OutputTensor, TConstArrayView<int32> Starts);
+	void NNERUNTIMERDG_API Apply(const NNE::Internal::FTensor& InputTensor, NNE::Internal::FTensor& OutputTensor, TConstArrayView<int32> Starts);
 } // UE::NNERuntimeRDG::Internal::CPUHelper::Slice
\ No newline at end of file
diff --git a/Engine/Plugins/Experimental/NNERuntimeRDG/Source/NNERuntimeRDG/Internal/NNERuntimeRDGHelperTranspose.h b/Engine/Plugins/Experimental/NNERuntimeRDG/Source/NNERuntimeRDG/Internal/NNERuntimeRDGHelperTranspose.h
--- a/Engine/Plugins/Experimental/NNERuntimeRDG/Source/NNERuntimeRDG/Internal/NNERuntimeRDGHelperTranspose.h
+++ b/Engine/Plugins/Experimental/NNERuntimeRDG/Source/NNERuntimeRDG/Internal/NNERuntimeRDGHelperTranspose.h
@@ -4,9 +4,9 @@
 
 #include "CoreMinimal.h"
 
-namespace UE::NNECore::Internal { class FTensor; }
+namespace UE::NNE::Internal { class FTensor; }
 
 namespace UE::NNERuntimeRDG::Internal::CPUHelper::Transpose
 {
-	bool NNERUNTIMERDG_API TransposePreparedData(NNECore::Internal::FTensor& Tensor, TConstArrayView<int32> Perms);
+	bool NNERUNTIMERDG_API TransposePreparedData(NNE::Internal::FTensor& Tensor, TConstArrayView<int32> Perms);
 } // UE::NNERuntimeRDG::Internal::CPUHelper::Transpose
\ No newline at end of file
