[ML Deformer] Fixed issues with the VertexDeltaModel which caused animation to appear frozen

diff --git a/Engine/Plugins/Animation/MLDeformer/VertexDeltaModel/Source/VertexDeltaModel/Private/VertexDeltaGraphDataInterface.cpp b/Engine/Plugins/Animation/MLDeformer/VertexDeltaModel/Source/VertexDeltaModel/Private/VertexDeltaGraphDataInterface.cpp
--- a/Engine/Plugins/Animation/MLDeformer/VertexDeltaModel/Source/VertexDeltaModel/Private/VertexDeltaGraphDataInterface.cpp
+++ b/Engine/Plugins/Animation/MLDeformer/VertexDeltaModel/Source/VertexDeltaModel/Private/VertexDeltaGraphDataInterface.cpp
@@ -119,7 +119,6 @@ namespace UE::VertexDeltaModel
 
 				if (NNEModelInstanceRDG)
 				{
-					const FVertexMapBuffer& VertexMapBuffer = VertexDeltaModel->GetVertexMapBuffer();
 					Weight = DeformerComponent->GetWeight();
 					VertexMapBufferSRV = VertexDeltaModel->GetVertexMapBuffer().ShaderResourceViewRHI;
 				}
@@ -134,14 +133,14 @@ namespace UE::VertexDeltaModel
 			return false;
 		}
 
-		if (!bCanRunNeuralNet || SkeletalMeshObject == nullptr || NNEModelInstanceRDG == nullptr || VertexMapBufferSRV == nullptr)
+		if (SkeletalMeshObject == nullptr || NNEModelInstanceRDG == nullptr || VertexMapBufferSRV == nullptr)
 		{
 			return false;
 		}
 
 		return true;
 	}
-
+	
 
 	void FVertexDeltaGraphDataProviderProxy::AllocateResources(FRDGBuilder& GraphBuilder)
 	{
diff --git a/Engine/Plugins/Animation/MLDeformer/VertexDeltaModel/Source/VertexDeltaModel/Public/VertexDeltaGraphDataInterface.h b/Engine/Plugins/Animation/MLDeformer/VertexDeltaModel/Source/VertexDeltaModel/Public/VertexDeltaGraphDataInterface.h
--- a/Engine/Plugins/Animation/MLDeformer/VertexDeltaModel/Source/VertexDeltaModel/Public/VertexDeltaGraphDataInterface.h
+++ b/Engine/Plugins/Animation/MLDeformer/VertexDeltaModel/Source/VertexDeltaModel/Public/VertexDeltaGraphDataInterface.h
@@ -92,6 +92,5 @@ namespace UE::VertexDeltaModel
 		FRDGBufferSRV* BufferSRV = nullptr;
 		int32 NeuralNetworkInferenceHandle = -1;
 		float Weight = 1.0f;
-		bool bCanRunNeuralNet = false;
 	};
 }	// namespace UE::VertexDeltaModel
diff --git a/Engine/Plugins/Animation/MLDeformer/VertexDeltaModel/VertexDeltaModel.uplugin b/Engine/Plugins/Animation/MLDeformer/VertexDeltaModel/VertexDeltaModel.uplugin
--- a/Engine/Plugins/Animation/MLDeformer/VertexDeltaModel/VertexDeltaModel.uplugin
+++ b/Engine/Plugins/Animation/MLDeformer/VertexDeltaModel/VertexDeltaModel.uplugin
@@ -42,6 +42,10 @@
 			"Name": "NNE",
 			"Enabled": true
 		},
+		{
+			"Name": "NNERuntimeRDG",
+			"Enabled": true
+		},
 		{
 			"Name": "DeformerGraph",
 			"Enabled": true
