[MLDeformer]
- Some small tweak when detecting if an animation is the same length or not.
- Fix in the Vertex Delta Model training, where it would fail when clicking Train, as the API changed but the python script wasn't updated yet.
- Fixed a bug in the VertexDeltaModel where it wouldn't save the neural network.
- Fixed a bug where the VertexDeltaModel extracted the wrong size from a tensor.
- Some other minor code changes.

diff --git a/Engine/Plugins/Animation/MLDeformer/MLDeformerFramework/Source/MLDeformerFramework/Private/MLDeformerGeomCacheHelpers.cpp b/Engine/Plugins/Animation/MLDeformer/MLDeformerFramework/Source/MLDeformerFramework/Private/MLDeformerGeomCacheHelpers.cpp
--- a/Engine/Plugins/Animation/MLDeformer/MLDeformerFramework/Source/MLDeformerFramework/Private/MLDeformerGeomCacheHelpers.cpp
+++ b/Engine/Plugins/Animation/MLDeformer/MLDeformerFramework/Source/MLDeformerFramework/Private/MLDeformerGeomCacheHelpers.cpp
@@ -368,7 +368,7 @@ namespace UE::MLDeformer
 					FText::AsNumber(NumGeomCacheFrames));
 			}
 
-			if (FMath::Abs(AnimSeqDuration - GeomCacheDuration) > 0.0001f)
+			if (FMath::Abs(AnimSeqDuration - GeomCacheDuration) > 0.001f)
 			{
 				FNumberFormattingOptions Options;
 				Options.SetUseGrouping(false);
diff --git a/Engine/Plugins/Animation/MLDeformer/VertexDeltaModel/Source/VertexDeltaModel/Private/VertexDeltaGraphDataInterface.cpp b/Engine/Plugins/Animation/MLDeformer/VertexDeltaModel/Source/VertexDeltaModel/Private/VertexDeltaGraphDataInterface.cpp
--- a/Engine/Plugins/Animation/MLDeformer/VertexDeltaModel/Source/VertexDeltaModel/Private/VertexDeltaGraphDataInterface.cpp
+++ b/Engine/Plugins/Animation/MLDeformer/VertexDeltaModel/Source/VertexDeltaModel/Private/VertexDeltaGraphDataInterface.cpp
@@ -104,9 +104,9 @@ namespace UE::VertexDeltaModel
 	{
 		using namespace UE::MLDeformer;
 
-		const UMLDeformerAsset* DeformerAsset = DeformerComponent != nullptr ? DeformerComponent->GetDeformerAsset() : nullptr;
-		const UMLDeformerModel* Model = DeformerAsset != nullptr ? DeformerAsset->GetModel() : nullptr;
-		UMLDeformerModelInstance* ModelInstance = DeformerComponent != nullptr ? DeformerComponent->GetModelInstance() : nullptr;
+		const UMLDeformerAsset* DeformerAsset = DeformerComponent ? DeformerComponent->GetDeformerAsset() : nullptr;
+		const UMLDeformerModel* Model = DeformerAsset ? DeformerAsset->GetModel() : nullptr;
+		UMLDeformerModelInstance* ModelInstance = DeformerComponent ? DeformerComponent->GetModelInstance() : nullptr;
 		const UVertexDeltaModel* VertexDeltaModel = Cast<UVertexDeltaModel>(Model);
 		
 		if (Model && VertexDeltaModel && ModelInstance)
diff --git a/Engine/Plugins/Animation/MLDeformer/VertexDeltaModel/Source/VertexDeltaModel/Private/VertexDeltaModel.cpp b/Engine/Plugins/Animation/MLDeformer/VertexDeltaModel/Source/VertexDeltaModel/Private/VertexDeltaModel.cpp
--- a/Engine/Plugins/Animation/MLDeformer/VertexDeltaModel/Source/VertexDeltaModel/Private/VertexDeltaModel.cpp
+++ b/Engine/Plugins/Animation/MLDeformer/VertexDeltaModel/Source/VertexDeltaModel/Private/VertexDeltaModel.cpp
@@ -5,6 +5,7 @@
 #include "VertexDeltaModelInstance.h"
 #include "MLDeformerComponent.h"
 #include "MLDeformerAsset.h"
+#include "MLDeformerObjectVersion.h"
 #include "NNE.h"
 #include "NNEModelData.h"
 #include "NNERuntimeRDG.h"
@@ -62,11 +63,6 @@ bool UVertexDeltaModel::IsTrained() const
 	return NNEModel.Get() != nullptr;
 }
 
-void UVertexDeltaModel::PostLoad()
-{
-	Super::PostLoad();
-}
-
 void UVertexDeltaModel::GetAssetRegistryTags(TArray<FAssetRegistryTag>& OutTags) const
 {
 	Super::GetAssetRegistryTags(OutTags);
diff --git a/Engine/Plugins/Animation/MLDeformer/VertexDeltaModel/Source/VertexDeltaModel/Private/VertexDeltaModelInstance.cpp b/Engine/Plugins/Animation/MLDeformer/VertexDeltaModel/Source/VertexDeltaModel/Private/VertexDeltaModelInstance.cpp
--- a/Engine/Plugins/Animation/MLDeformer/VertexDeltaModel/Source/VertexDeltaModel/Private/VertexDeltaModelInstance.cpp
+++ b/Engine/Plugins/Animation/MLDeformer/VertexDeltaModel/Source/VertexDeltaModel/Private/VertexDeltaModelInstance.cpp
@@ -11,7 +11,6 @@ UVertexDeltaModel* UVertexDeltaModelInstance::GetVertexDeltaModel() const
 	return Cast<UVertexDeltaModel>(Model);
 }
 
-
 UE::NNE::IModelInstanceRDG* UVertexDeltaModelInstance::GetNNEModelInstanceRDG() const
 {
 	return ModelInstanceRDG.Get();
@@ -31,7 +30,7 @@ FString UVertexDeltaModelInstance::CheckCompatibility(USkeletalMeshComponent* In
 	if (RDGModelInstance && Model->GetDeformerAsset())
 	{
 		TConstArrayView<UE::NNE::FTensorDesc> GPUTensorDesc = RDGModelInstance->GetInputTensorDescs();
-		const int64 NumNeuralNetInputs = GPUTensorDesc[0].GetShape().Rank();
+		const int64 NumNeuralNetInputs = GPUTensorDesc[0].GetShape().GetData()[1];
 		const int64 NumDeformerAssetInputs = static_cast<int64>(Model->GetInputInfo()->CalcNumNeuralNetInputs(Model->GetNumFloatsPerBone(), Model->GetNumFloatsPerCurve()));
 		if (NumNeuralNetInputs != NumDeformerAssetInputs)
 		{
@@ -54,7 +53,7 @@ bool UVertexDeltaModelInstance::IsValidForDataProvider() const
 
 void UVertexDeltaModelInstance::Execute(float ModelWeight)
 {
-	TRACE_CPUPROFILER_EVENT_SCOPE(UMLDeformerModelInstance::Execute)
+	TRACE_CPUPROFILER_EVENT_SCOPE(UVertexDeltaModelInstance::Execute)
 
 	if (ModelInstanceRDG)
 	{
@@ -119,7 +118,8 @@ bool UVertexDeltaModelInstance::GetRDGVertexBufferDesc(TConstArrayView<UE::NNE::
 
 void UVertexDeltaModelInstance::CreateRDGBuffers(TConstArrayView<UE::NNE::FTensorDesc>& OutputTensorDescs)
 {
-	ENQUEUE_RENDER_COMMAND(VertexDeltaModelInstance_CreateOuputRDGBuffer)(
+	ENQUEUE_RENDER_COMMAND(VertexDeltaModelInstance_CreateOuputRDGBuffer)
+	(
 		[this, &OutputTensorDescs](FRHICommandListImmediate& RHICmdList)
 		{
 			FRDGBuilder Builder(RHICmdList);
@@ -143,7 +143,6 @@ void UVertexDeltaModelInstance::CreateRDGBuffers(TConstArrayView<UE::NNE::FTenso
 	FRenderCommandFence RenderFence;
 	RenderFence.BeginFence();
 	RenderFence.Wait();
-
 }
 
 void UVertexDeltaModelInstance::PostMLDeformerComponentInit()
@@ -167,7 +166,7 @@ void UVertexDeltaModelInstance::CreateNNEModel()
 
 			if (!Runtime.IsValid())
 			{
-				UE_LOG(LogNNE, Error, TEXT("Can't get %s runtime."), *VertexDeltaModel->GetNNERuntimeName());
+				UE_LOG(LogVertexDeltaModel, Error, TEXT("Can't get NNE runtime: %s"), *VertexDeltaModel->GetNNERuntimeName());
 				return;
 			}
 
@@ -194,7 +193,7 @@ void UVertexDeltaModelInstance::CreateNNEModel()
 					}
 					else
 					{
-						UE_LOG(LogNNE, Error, TEXT("Failed to create NNE RDG Model for VertexDeltaModel."));
+						UE_LOG(LogVertexDeltaModel, Error, TEXT("Failed to create NNE RDG Model for VertexDeltaModel."));
 						return;
 					}
 				}
@@ -205,19 +204,18 @@ void UVertexDeltaModelInstance::CreateNNEModel()
 
 bool UVertexDeltaModelInstance::SetupInputs()
 {
-	TRACE_CPUPROFILER_EVENT_SCOPE(UMLDeformerModelInstance::SetupInputs)
+	TRACE_CPUPROFILER_EVENT_SCOPE(UVertexDeltaModelInstance::SetupInputs)
 
 	// Some safety checks.
-	if (Model == nullptr ||
-		SkeletalMeshComponent == nullptr ||
-		SkeletalMeshComponent->GetSkeletalMeshAsset() == nullptr ||
+	if (!Model ||
+		!SkeletalMeshComponent ||
+		!SkeletalMeshComponent->GetSkeletalMeshAsset() ||
 		!bIsCompatible)
 	{
 		return false;
 	}
 
 	// Get the network and make sure it's loaded.
-
 	if (!ModelInstanceRDG)
 	{
 		return false;
@@ -236,6 +234,7 @@ bool UVertexDeltaModelInstance::SetupInputs()
 	{
 		return false; 
 	}
+
 	// Update and write the input values directly into the input tensor.
 	const int64 NumFloatsWritten = SetNeuralNetworkInputValues(NNEInputTensorBuffer.GetData(), NNEInputTensorBuffer.Num());
 	check(NumFloatsWritten == NumNeuralNetInputs);
diff --git a/Engine/Plugins/Animation/MLDeformer/VertexDeltaModel/Source/VertexDeltaModel/Public/VertexDeltaModel.h b/Engine/Plugins/Animation/MLDeformer/VertexDeltaModel/Source/VertexDeltaModel/Public/VertexDeltaModel.h
--- a/Engine/Plugins/Animation/MLDeformer/VertexDeltaModel/Source/VertexDeltaModel/Public/VertexDeltaModel.h
+++ b/Engine/Plugins/Animation/MLDeformer/VertexDeltaModel/Source/VertexDeltaModel/Public/VertexDeltaModel.h
@@ -29,7 +29,6 @@ public:
 	UVertexDeltaModel(const FObjectInitializer& ObjectInitializer);
 
 	// UObject overrides.
-	virtual void PostLoad() override;
 	virtual void GetAssetRegistryTags(TArray<FAssetRegistryTag>& OutTags) const override;
 	// ~END UObject overrides.
 
@@ -48,13 +47,14 @@ public:
 	int32 GetBatchSize() const								{ return BatchSize; }
 	float GetLearningRate() const							{ return LearningRate; }
 #endif
+
 	void SetNNEModelData(TObjectPtr<UNNEModelData> ModelData);
-	const FString GetNNERuntimeName() const { return TEXT("NNERuntimeRDGDml"); }
-public:
+	const FString GetNNERuntimeName() const					{ return TEXT("NNERuntimeRDGDml"); }
 
+public:
 	/** The NNE neural network model. */
 	UPROPERTY()
-	TObjectPtr<UNNEModelData> NNEModel; 
+	TObjectPtr<UNNEModelData> NNEModel;
 
 #if WITH_EDITORONLY_DATA
 	/** The number of hidden layers that the neural network model will have.\nHigher numbers will slow down performance but can deal with more complex deformations. */
@@ -77,5 +77,4 @@ public:
 	UPROPERTY(EditAnywhere, BlueprintReadWrite, Category = "Training Settings", meta = (ClampMin = "0.000001", ClampMax = "1.0"))
 	float LearningRate = 0.001f;
 #endif // WITH_EDITORONLY_DATA
-
 };
diff --git a/Engine/Plugins/Animation/MLDeformer/VertexDeltaModel/Source/VertexDeltaModelEditor/Private/VertexDeltaEditorModel.cpp b/Engine/Plugins/Animation/MLDeformer/VertexDeltaModel/Source/VertexDeltaModelEditor/Private/VertexDeltaEditorModel.cpp
--- a/Engine/Plugins/Animation/MLDeformer/VertexDeltaModel/Source/VertexDeltaModelEditor/Private/VertexDeltaEditorModel.cpp
+++ b/Engine/Plugins/Animation/MLDeformer/VertexDeltaModel/Source/VertexDeltaModelEditor/Private/VertexDeltaEditorModel.cpp
@@ -31,7 +31,7 @@ namespace UE::VertexDeltaModel
 		return TrainModel<UVertexDeltaTrainingModel>(this);
 	}
 
-	TObjectPtr<UNNEModelData>  FVertexDeltaEditorModel::LoadNeuralNetworkFromOnnx(const FString& Filename) const
+	TObjectPtr<UNNEModelData> FVertexDeltaEditorModel::LoadNeuralNetworkFromOnnx(const FString& Filename) const
 	{
 		const FString OnnxFile = FPaths::ConvertRelativePathToFull(Filename);
 		if (FPaths::FileExists(OnnxFile))
@@ -42,11 +42,9 @@ namespace UE::VertexDeltaModel
 	
 			if (bIsModelInMem)
 			{
-				TObjectPtr<UNNEModelData> ModelData = NewObject<UNNEModelData>();
-
+				TObjectPtr<UNNEModelData> ModelData = NewObject<UNNEModelData>(Model);
 				ModelData->Init(FString("onnx"), RawModelData);
-				// TODO - add some real verification
-				UE_LOG(LogVertexDeltaModel, Display, TEXT("Successfully loaded Onnx file '%s'..."), *OnnxFile);
+				UE_LOG(LogVertexDeltaModel, Display, TEXT("Finished loading Onnx file '%s'..."), *OnnxFile);
 				return ModelData;
 			}
 			else
@@ -65,10 +63,10 @@ namespace UE::VertexDeltaModel
 	bool FVertexDeltaEditorModel::LoadTrainedNetwork() const
 	{
 		const FString OnnxFile = GetTrainedNetworkOnnxFile();
-		TObjectPtr<UNNEModelData> ModelData = LoadNeuralNetworkFromOnnx(OnnxFile);
-		if (ModelData)
+		TObjectPtr<UNNEModelData> LoadedModelData = LoadNeuralNetworkFromOnnx(OnnxFile);
+		if (LoadedModelData)
 		{
-			GetVertexDeltaModel()->SetNNEModelData(ModelData);
+			GetVertexDeltaModel()->SetNNEModelData(LoadedModelData);
 			return true;
 		}
 
diff --git a/Engine/Plugins/Animation/MLDeformer/VertexDeltaModel/Source/VertexDeltaModelEditor/Public/VertexDeltaTrainingModel.h b/Engine/Plugins/Animation/MLDeformer/VertexDeltaModel/Source/VertexDeltaModelEditor/Public/VertexDeltaTrainingModel.h
--- a/Engine/Plugins/Animation/MLDeformer/VertexDeltaModel/Source/VertexDeltaModelEditor/Public/VertexDeltaTrainingModel.h
+++ b/Engine/Plugins/Animation/MLDeformer/VertexDeltaModel/Source/VertexDeltaModelEditor/Public/VertexDeltaTrainingModel.h
@@ -3,7 +3,7 @@
 #pragma once
 
 #include "CoreTypes.h"
-#include "MLDeformerTrainingModel.h"
+#include "MLDeformerGeomCacheTrainingModel.h"
 #include "VertexDeltaTrainingModel.generated.h"
 
 /**
@@ -12,7 +12,7 @@
  */
 UCLASS(Blueprintable)
 class VERTEXDELTAMODELEDITOR_API UVertexDeltaTrainingModel
-	: public UMLDeformerTrainingModel
+	: public UMLDeformerGeomCacheTrainingModel
 {
 	GENERATED_BODY()
 
