[MLDeformer] Fix a crash when saving a level after creating a character blueprint with Vertex Delta Model in it and a compute graph component.

diff --git a/Engine/Plugins/Animation/MLDeformer/VertexDeltaModel/Source/VertexDeltaModel/Private/VertexDeltaGraphDataInterface.cpp b/Engine/Plugins/Animation/MLDeformer/VertexDeltaModel/Source/VertexDeltaModel/Private/VertexDeltaGraphDataInterface.cpp
--- a/Engine/Plugins/Animation/MLDeformer/VertexDeltaModel/Source/VertexDeltaModel/Private/VertexDeltaGraphDataInterface.cpp
+++ b/Engine/Plugins/Animation/MLDeformer/VertexDeltaModel/Source/VertexDeltaModel/Private/VertexDeltaGraphDataInterface.cpp
@@ -109,7 +109,8 @@ namespace UE::VertexDeltaModel
 		UMLDeformerModelInstance* ModelInstance = DeformerComponent ? DeformerComponent->GetModelInstance() : nullptr;
 		const UVertexDeltaModel* VertexDeltaModel = Cast<UVertexDeltaModel>(Model);
 		
-		if (Model && VertexDeltaModel && ModelInstance)
+		Weight = 0.0f;
+		if (Model && VertexDeltaModel && ModelInstance && ModelInstance->GetSkeletalMeshComponent())
 		{
 			VertexDeltaModelInstance = Cast<UVertexDeltaModelInstance>(ModelInstance);
 			if (VertexDeltaModelInstance)
