[MLDeformer] Fixed a crash when packaging a project with a blueprint that has a VertexDeltaModel based ML Deformer inside it.

diff --git a/Engine/Plugins/Animation/MLDeformer/VertexDeltaModel/Source/VertexDeltaModel/Private/VertexDeltaModelInstance.cpp b/Engine/Plugins/Animation/MLDeformer/VertexDeltaModel/Source/VertexDeltaModel/Private/VertexDeltaModelInstance.cpp
--- a/Engine/Plugins/Animation/MLDeformer/VertexDeltaModel/Source/VertexDeltaModel/Private/VertexDeltaModelInstance.cpp
+++ b/Engine/Plugins/Animation/MLDeformer/VertexDeltaModel/Source/VertexDeltaModel/Private/VertexDeltaModelInstance.cpp
@@ -163,31 +163,31 @@ void UVertexDeltaModelInstance::CreateNNEModel()
 		{
 			TWeakInterfacePtr<INNERuntime> Runtime = UE::NNE::GetRuntime<INNERuntime>(VertexDeltaModel->GetNNERuntimeName());
 			TWeakInterfacePtr<INNERuntimeRDG> RuntimeRDG = UE::NNE::GetRuntime<INNERuntimeRDG>(VertexDeltaModel->GetNNERuntimeName());
-
 			if (!Runtime.IsValid())
 			{
 				UE_LOG(LogVertexDeltaModel, Error, TEXT("Can't get NNE runtime: %s"), *VertexDeltaModel->GetNNERuntimeName());
 				return;
 			}
 
+			// If we can create the model from its data.
 			TObjectPtr<UNNEModelData> ModelData = VertexDeltaModel->NNEModel;
-			if (ModelData)
+			if (ModelData && RuntimeRDG.IsValid() && RuntimeRDG->CanCreateModelRDG(ModelData) == INNERuntimeRDG::ECanCreateModelRDGStatus::Ok)
 			{
-				if (RuntimeRDG.IsValid())
+				// Create the model.
+				TSharedPtr<UE::NNE::IModelRDG> ModelRDG = RuntimeRDG->CreateModelRDG(ModelData);
+				if (ModelRDG.IsValid())
 				{
-					// allocate tensor inputs and outputs
 					ModelInstanceRDG = RuntimeRDG->CreateModelRDG(ModelData)->CreateModelInstanceRDG();
-
 					if (ModelInstanceRDG)
 					{
-						// setup inputs
+						// Setup inputs.
 						TConstArrayView<UE::NNE::FTensorDesc> InputTensorDescs = ModelInstanceRDG->GetInputTensorDescs();
 						UE::NNE::FTensorShape InputTensorShape = UE::NNE::FTensorShape::MakeFromSymbolic(InputTensorDescs[0].GetShape());
 						ModelInstanceRDG->SetInputTensorShapes({ InputTensorShape });
 						check(InputTensorDescs[0].GetElementByteSize() == sizeof(float));
 						NNEInputTensorBuffer.SetNumUninitialized(InputTensorShape.Volume());
 
-						// setup outputs
+						// Setup outputs.
 						TConstArrayView<UE::NNE::FTensorDesc> OutputTensorDescs = ModelInstanceRDG->GetOutputTensorDescs();
 						CreateRDGBuffers(OutputTensorDescs);
 					}
