[MLDeformer] Vertex delta model was still using the NNE RDG instead of ORT, which caused it to not be able to run the model properly.

diff --git a/Engine/Plugins/Animation/MLDeformer/VertexDeltaModel/Source/VertexDeltaModel/VertexDeltaModel.Build.cs b/Engine/Plugins/Animation/MLDeformer/VertexDeltaModel/Source/VertexDeltaModel/VertexDeltaModel.Build.cs
--- a/Engine/Plugins/Animation/MLDeformer/VertexDeltaModel/Source/VertexDeltaModel/VertexDeltaModel.Build.cs
+++ b/Engine/Plugins/Animation/MLDeformer/VertexDeltaModel/Source/VertexDeltaModel/VertexDeltaModel.Build.cs
@@ -39,7 +39,7 @@ public class VertexDeltaModel : ModuleRules
 				"Engine",
 				"GeometryCache",
 				"NNE",
-				"NNERuntimeRDG",
+				"NNERuntimeORT",
 				"OptimusCore",
 				"Projects",
 				"RenderCore",
diff --git a/Engine/Plugins/Animation/MLDeformer/VertexDeltaModel/Source/VertexDeltaModelEditor/VertexDeltaModelEditor.Build.cs b/Engine/Plugins/Animation/MLDeformer/VertexDeltaModel/Source/VertexDeltaModelEditor/VertexDeltaModelEditor.Build.cs
--- a/Engine/Plugins/Animation/MLDeformer/VertexDeltaModel/Source/VertexDeltaModelEditor/VertexDeltaModelEditor.Build.cs
+++ b/Engine/Plugins/Animation/MLDeformer/VertexDeltaModel/Source/VertexDeltaModelEditor/VertexDeltaModelEditor.Build.cs
@@ -32,7 +32,7 @@ namespace UnrealBuildTool.Rules
 					"VertexDeltaModel",
 					"PropertyEditor",
 					"NNE",
-					"NNERuntimeRDG",
+					"NNERuntimeORT",
 					"ToolWidgets",
 					"ComputeFramework"
 				}
diff --git a/Engine/Plugins/Animation/MLDeformer/VertexDeltaModel/VertexDeltaModel.uplugin b/Engine/Plugins/Animation/MLDeformer/VertexDeltaModel/VertexDeltaModel.uplugin
--- a/Engine/Plugins/Animation/MLDeformer/VertexDeltaModel/VertexDeltaModel.uplugin
+++ b/Engine/Plugins/Animation/MLDeformer/VertexDeltaModel/VertexDeltaModel.uplugin
@@ -39,7 +39,7 @@
 			"Enabled": true
 		},
 		{
-			"Name": "NNERuntimeRDG",
+			"Name": "NNERuntimeORT",
 			"Enabled": true
 		},
 		{
